#UE4 - Removing game references from distributed code (TTP# 330940). Deleted unnecessary batch files. Moved TestGUBPLocal.bat back to BatchFiles - it no longer contains any game references.

diff --git a/Engine/Build/BatchFiles/TestGUBPLocal.bat b/Engine/Build/BatchFiles/TestGUBPLocal.bat
new file mode 100644
--- /dev/null
+++ b/Engine/Build/BatchFiles/TestGUBPLocal.bat
@@ -0,0 +1,29 @@
+set uebp_testfail=0
+
+call runuattest.bat gubp -ListOnly -cleanlocal -Allplatforms -Fake -FakeEC %*
+
+call runuattest.bat gubp -Allplatforms -Fake -FakeEC %* -CommanderJobSetupOnly
+
+call runuattest.bat gubp -Allplatforms -Fake -FakeEC %* -CommanderJobSetupOnly -TriggerNode=Rocket_WaitToMakeBuild
+
+call runuattest.bat gubp -Allplatforms -Fake -FakeEC %* -CommanderJobSetupOnly -TriggerNode=Shared_WaitForPromotable
+
+call runuattest.bat gubp -Allplatforms -Fake -FakeEC %* -CommanderJobSetupOnly -TriggerNode=Shared_WaitForTesting
+
+call runuattest.bat gubp -Allplatforms -Fake -FakeEC %* -CommanderJobSetupOnly -TriggerNode=Shared_WaitForPromotion
+
+
+
+
+if "%uebp_testfail%" == "1" goto fail
+
+ECHO ****************** All tests Succeeded
+exit /B 0
+
+:fail
+ECHO ****************** A test failed
+exit /B 1
+
+
+
+	
\ No newline at end of file
