#ue4 Fix for registry error when running builds while logged in as SYSTEM.  See: https://udn.unrealengine.com/questions/204468/copyvisualizersbat-change-request.html

#fix Skipped the copying if registry lookup fails.

diff --git a/Engine/Build/BatchFiles/CopyVisualizers.bat b/Engine/Build/BatchFiles/CopyVisualizers.bat
--- a/Engine/Build/BatchFiles/CopyVisualizers.bat
+++ b/Engine/Build/BatchFiles/CopyVisualizers.bat
@@ -1,17 +1,21 @@
 @echo off
 
-for /f "tokens=2,*" %%a in ('reg query "HKEY_CURRENT_USER\Software\Microsoft\Windows\CurrentVersion\Explorer\Shell Folders" /v Personal ^| findstr Personal') do (
-	set UE4_MyDocs=%%b
-	
-	if "%VS120COMNTOOLS%" == "" goto NoVisualStudio2013Environment
-	attrib -R "!UE4_MyDocs!\Visual Studio 2013\Visualizers\UE4.natvis" 1>nul 2>nul
-	copy /Y "%~dp0..\..\Extras\VisualStudioDebugging\UE4.natvis" "!UE4_MyDocs!\Visual Studio 2013\Visualizers\UE4.natvis" 1>nul 2>nul
+REM We don't want to copy visualizers if the key lookup fails, as it may do with a build system running as SYSTEM
+reg query "HKEY_CURRENT_USER\Software\Microsoft\Windows\CurrentVersion\Explorer\Shell Folders" /v Personal 1>nul 2>&1
+if errorlevel 0 (
+	for /f "tokens=2,*" %%a in ('reg query "HKEY_CURRENT_USER\Software\Microsoft\Windows\CurrentVersion\Explorer\Shell Folders" /v Personal ^| findstr Personal') do (
+		set UE4_MyDocs=%%b
+		
+		if "%VS120COMNTOOLS%" == "" goto NoVisualStudio2013Environment
+		attrib -R "!UE4_MyDocs!\Visual Studio 2013\Visualizers\UE4.natvis" 1>nul 2>nul
+		copy /Y "%~dp0..\..\Extras\VisualStudioDebugging\UE4.natvis" "!UE4_MyDocs!\Visual Studio 2013\Visualizers\UE4.natvis" 1>nul 2>nul
 
-:NoVisualStudio2013Environment
-	if "%VS110COMNTOOLS%" == "" goto NoVisualStudio2012Environment
-	attrib -R "!UE4_MyDocs!\Visual Studio 2012\Visualizers\UE4.natvis" 1>nul 2>nul
-	copy /Y "%~dp0..\..\Extras\VisualStudioDebugging\UE4.natvis" "!UE4_MyDocs!\Visual Studio 2012\Visualizers\UE4.natvis" 1>nul 2>nul
+	:NoVisualStudio2013Environment
+		if "%VS110COMNTOOLS%" == "" goto NoVisualStudio2012Environment
+		attrib -R "!UE4_MyDocs!\Visual Studio 2012\Visualizers\UE4.natvis" 1>nul 2>nul
+		copy /Y "%~dp0..\..\Extras\VisualStudioDebugging\UE4.natvis" "!UE4_MyDocs!\Visual Studio 2012\Visualizers\UE4.natvis" 1>nul 2>nul
 
-:NoVisualStudio2012Environment
-	set UE4_MyDocs=
+	:NoVisualStudio2012Environment
+		set UE4_MyDocs=
+	)
 )
