Fixed GenerateProjectfiles.bat failing to fallback to VS2012 if VS2013 is not located.

diff --git a/Engine/Build/BatchFiles/GenerateProjectFiles.bat b/Engine/Build/BatchFiles/GenerateProjectFiles.bat
--- a/Engine/Build/BatchFiles/GenerateProjectFiles.bat
+++ b/Engine/Build/BatchFiles/GenerateProjectFiles.bat
@@ -37,15 +37,14 @@ if "%VsComnToolsPath%" == "" goto NoVisualStudio2013Environment
 call "%VsComnToolsPath%/../../VC/bin/x86_amd64/vcvarsx86_amd64.bat" >NUL
 goto ReadyToCompile
 
-
+:NoVisualStudio2013Environment
 rem ## Check for Visual Studio 2012
 
 pushd %~dp0
 call GetVSComnToolsPath 11
 popd
 
-:NoVisualStudio2013Environment
-if "%VsComnToolsPath%" == "" goto NoVisualStudioEnvironment
+if "%VsComnToolsPath%" == "" goto Error_NoVisualStudioEnvironment
 call "%VsComnToolsPath%/../../VC/bin/x86_amd64/vcvarsx86_amd64.bat" >NUL
 goto ReadyToCompile
 
