Fixed Clean.bat failing to fallback to VS2012 if VS2013 is not located. Matched similar structure in GenerateProjectFiles.bat to Clean.bat.

diff --git a/Engine/Build/BatchFiles/GenerateProjectFiles.bat b/Engine/Build/BatchFiles/GenerateProjectFiles.bat
--- a/Engine/Build/BatchFiles/GenerateProjectFiles.bat
+++ b/Engine/Build/BatchFiles/GenerateProjectFiles.bat
@@ -44,11 +44,12 @@ pushd %~dp0
 call GetVSComnToolsPath 11
 popd
 
-if "%VsComnToolsPath%" == "" goto Error_NoVisualStudioEnvironment
+if "%VsComnToolsPath%" == "" goto NoVisualStudio2012Environment
 call "%VsComnToolsPath%/../../VC/bin/x86_amd64/vcvarsx86_amd64.bat" >NUL
 goto ReadyToCompile
 
 
+:NoVisualStudio2012Environment
 rem ## User has no version of Visual Studio installed?
 goto Error_NoVisualStudioEnvironment
 
