fix for strange label not found error in RunUAT.bat
#ue4
#uat

diff --git a/Engine/Build/BatchFiles/RunUAT.bat b/Engine/Build/BatchFiles/RunUAT.bat
--- a/Engine/Build/BatchFiles/RunUAT.bat
+++ b/Engine/Build/BatchFiles/RunUAT.bat
@@ -20,6 +20,7 @@ if not exist Build\BatchFiles\RunUAT.bat goto Error_BatchFileInWrongLocation
 rem ## Use the pre-compiled UAT scripts if -nocompile is specified in the command line
 for %%P in (%*) do if /I "%%P" == "-nocompile" goto RunPrecompiled
 
+rem ## check for force precompiled
 if not "%ForcePrecompiledUAT%"=="" goto RunPrecompiled
 
 if not exist Source\Programs\AutomationTool\AutomationTool.csproj goto RunPrecompiled
