After merging pretest with main calling CopyVisualizers.bat was made using EnableDelayedExpansion and since the variables reference had to be changed from '%' to '!'. This is the fix. Probably not found before cause in most cases visualizers were in place already.

diff --git a/Engine/Build/BatchFiles/CopyVisualizers.bat b/Engine/Build/BatchFiles/CopyVisualizers.bat
--- a/Engine/Build/BatchFiles/CopyVisualizers.bat
+++ b/Engine/Build/BatchFiles/CopyVisualizers.bat
@@ -9,8 +9,8 @@ if errorlevel 0 (
 		pushd %~dp0
 		call GetVSComnToolsPath 12
 		popd
-	
-		if "%VsComnToolsPath%" == "" goto NoVisualStudio2013Environment
+		
+		if "!VsComnToolsPath!" == "" goto NoVisualStudio2013Environment
 		attrib -R "!UE4_MyDocs!\Visual Studio 2013\Visualizers\UE4.natvis" 1>nul 2>nul
 		copy /Y "%~dp0..\..\Extras\VisualStudioDebugging\UE4.natvis" "!UE4_MyDocs!\Visual Studio 2013\Visualizers\UE4.natvis" 1>nul 2>nul
 
@@ -19,7 +19,7 @@ if errorlevel 0 (
 		call GetVSComnToolsPath 11
 		popd
 
-		if "%VsComnToolsPath%" == "" goto NoVisualStudio2012Environment
+		if "!VsComnToolsPath!" == "" goto NoVisualStudio2012Environment
 		attrib -R "!UE4_MyDocs!\Visual Studio 2012\Visualizers\UE4.natvis" 1>nul 2>nul
 		copy /Y "%~dp0..\..\Extras\VisualStudioDebugging\UE4.natvis" "!UE4_MyDocs!\Visual Studio 2012\Visualizers\UE4.natvis" 1>nul 2>nul
 
