Improvements to building TPS on Linux.

Part of PR140 by sbc100.
- Retouched to make more compliant with our coding style.
- Made to work out of the box with both git and perforce repos.
- Commented out TPS which we need to patch first.

diff --git a/Engine/Build/BatchFiles/Linux/Build.sh b/Engine/Build/BatchFiles/Linux/Build.sh
--- a/Engine/Build/BatchFiles/Linux/Build.sh
+++ b/Engine/Build/BatchFiles/Linux/Build.sh
@@ -7,11 +7,11 @@ set -e
 cd "`dirname "$0"`/../../../.." 
 
 # First make sure that the UnrealBuildTool is up-to-date
-if ! xbuild /property:Configuration=Development /verbosity:quiet /nologo Engine/Source/Programs/UnrealBuildTool/UnrealBuildTool_Mono.csproj; then
+if ! xbuild /property:Configuration=Development /property:TargetFrameworkVersion=v4.0 /verbosity:quiet /nologo Engine/Source/Programs/UnrealBuildTool/UnrealBuildTool_Mono.csproj; then
   echo "Failed to build to build tool (UnrealBuildTool)"
   exit 1
 fi
 
 echo "Building $1..."
 mono Engine/Binaries/DotNET/UnrealBuildTool.exe $*
-exit $?
\ No newline at end of file
+exit $?
diff --git a/Engine/Build/BatchFiles/Linux/BuildThirdParty.sh b/Engine/Build/BatchFiles/Linux/BuildThirdParty.sh
--- a/Engine/Build/BatchFiles/Linux/BuildThirdParty.sh
+++ b/Engine/Build/BatchFiles/Linux/BuildThirdParty.sh
@@ -1,23 +1,224 @@
 #!/bin/bash
+# Script for building ThirdParty modules for linux.
+# You must run InstallDeps.sh to download and patch the sources before running
+# this script.
 
 SCRIPT_DIR=$(cd "$(dirname "$BASH_SOURCE")" ; pwd)
 TOP_DIR=$(cd $SCRIPT_DIR/../../.. ; pwd)
 cd ${TOP_DIR}
 set -e
-pwd
-
-echo "building freetype"
-(
- cd Source/ThirdParty/FreeType2/FreeType2-2.4.12/src
- pwd
- make -j 5 -f ../Builds/Linux/makefile $*
- mkdir ../Lib/Linux
- cp libfreetype* ../Lib/Linux
-)
-
-echo "building recast"
-(
- cd Source/ThirdParty/Recast
- pwd
- make -j 5 -f Epic/build/Linux/makefile $*
-)
+
+MAKE_ARGS=-j4
+
+P4Open()
+{
+  set +x
+  which p4>>/dev/null && p4 open "$@"
+  set -x
+}
+
+BuildJemalloc()
+{
+  echo "building jemalloc"
+  set -x
+  cd Source/ThirdParty/jemalloc/build
+  tar xf jemalloc-3.4.0.tar.bz2
+  cd jemalloc-3.4.0
+  ./configure --with-mangling --with-jemalloc-prefix=je_340_
+  make
+  
+  P4Open ../../lib/Linux/separate/libjemalloc.a
+  cp lib/* ../../lib/Linux/separate/
+  cp lib/libjemalloc.so.1 ${TOP_DIR}/Binaries/Linux
+  set +x
+}
+
+BuildOpus()
+{
+  echo "building libOpus"
+  set -x
+  cd Source/ThirdParty/libOpus/opus-1.1/
+  P4Open configure
+  chmod +x configure
+  ./configure
+  make $MAKE_ARGS
+
+  cp .libs/libopus.a .libs/libopus*.so* Linux
+  cp Linux/libopus.so.0 ${TOP_DIR}/Binaries/Linux/
+  set +x
+}
+
+BuildPNG()
+{
+  echo "building libPNG"
+  set -x
+  cd Source/ThirdParty/libPNG/libPNG-1.5.2
+  cp --remove-destination scripts/makefile.linux Makefile
+  P4Open pnglibconf.h
+  make $MAKE_ARGS
+
+  P4Open lib/Linux/libpng.a
+  cp libpng.a lib/Linux
+  cp libpng15.so lib/Linux/libpng.so
+  if [ ! -d ../../../../Binaries/Linux ]; then
+    mkdir ../../../../Binaries/Linux
+  fi
+  cp -a libpng15.so* ${TOP_DIR}/Binaries/Linux/
+  set +x
+}
+
+BuildOgg()
+{
+  echo "building Ogg"
+  set -x
+  cd Source/ThirdParty/Ogg/libogg-1.2.2/
+  P4Open configure
+  chmod +x configure
+  ./configure
+  make $MAKE_ARGS
+
+  P4Open lib/Linux/libogg.a
+  cp src/.libs/* lib/Linux
+  set +x
+}
+
+BuildVorbis()
+{
+  echo "building Vorbis"
+  set -x
+  cd Source/ThirdParty/Vorbis/libvorbis-1.3.2/
+  P4Open configure
+  chmod +x configure
+  OGG_LIBS=../../Ogg/libogg-1.2.2/lib/Linux OGG_CFLAGS=-I../../../Ogg/libogg-1.2.2/include ./configure --with-pic --disable-oggtest
+  make $MAKE_ARGS
+
+  P4Open lib/Linux/libvorbis.a lib/Linux/libvorbisfile.a
+  cp lib/.libs/* lib/Linux
+  set +x
+  cd -
+}
+
+BuildHLSLCC()
+{
+  echo "building hlslcc"
+  set -x
+  cd Source/ThirdParty/hlslcc
+  make $MAKE_ARGS
+  set +x
+}
+
+BuildMcpp()
+{
+  echo "building MCPP"
+  set -x
+  cd Source/ThirdParty/MCPP/mcpp-2.7.2
+  ./configure --enable-shared --enable-shared --enable-mcpplib
+  make $MAKE_ARGS
+  rm -rf lib/Linux
+  mkdir lib/Linux
+  cp ./src/.libs/libmcpp.a lib/Linux
+  cp ./src/.libs/libmcpp.so lib/Linux
+  cp ./src/.libs/libmcpp.so ${TOP_DIR}/Binaries/Linux/libmcpp.so.0
+  set +x
+}
+
+BuildFreeType()
+{
+  echo "building freetype"
+  set -x
+  cd Source/ThirdParty/FreeType2/FreeType2-2.4.12/src
+  pwd
+  make $MAKE_ARGS -f ../Builds/Linux/makefile $*
+
+  P4Open ../Lib/Linux/libfreetype2412.a
+  cp --remove-destination libfreetype* ../Lib/Linux
+  set +x
+}
+
+BuildRecast()
+{
+  echo "building Recast"
+  set -x
+  cd Source/ThirdParty/Recast
+  pwd
+  P4Open Epic/lib/Linux/librecast.a
+  make $MAKE_ARGS -f Epic/build/Linux/makefile $*
+  set +x
+}
+
+BuildICU()
+{
+  echo "building libICU"
+  set -x
+  cd Source/ThirdParty/ICU/icu4c-51_2/source
+  P4Open configure
+  chmod +x configure
+  CPPFLAGS=-fPIC ./configure --enable-static --with-data-packaging=static
+  pwd
+  make $MAKE_ARGS $*
+
+  P4Open ../Linux/libicudata.a ../Linux/libicudatad.a ../Linux/libicui18n.a ../Linux/libicui18nd.a
+  P4Open ../Linux/libicuio.a ../Linux/libicuiod.a ../Linux/libicule.a ../Linux/libiculed.a
+  P4Open ../Linux/libiculx.a ../Linux/libiculxd.a ../Linux/libicutu.a ../Linux/libicutud.a
+  P4Open ../Linux/libicuuc.a ../Linux/libicuucd.a
+  cp --remove-destination lib/*.a ../Linux
+  set +x
+}
+
+BuildSDL2()
+{
+  echo "building dsext"
+  set -x
+  cd Source/ThirdParty/SDL2/DisplayServerExtensions
+  P4Open bin/libdsext.a
+  make $MAKE_ARGS
+  set +x
+  echo "building SDL2"
+  set -x
+  cd ../SDL2-2.0.3/build
+  pwd
+  tar xf SDL2-2.0.3.tar.gz
+  cd SDL2-2.0.3
+  P4Open configure
+  chmod +x configure
+  ./configure
+  make $MAKE_ARGS
+
+  # FIXME: build libdsext?
+  P4Open ../../lib/Linux/libEpicSDL2.a  
+  cp --remove-destination build/.libs/libSDL2.a ../../lib/Linux/libEpicSDL2.a
+  cp build/.libs/libSDL2.so ../../lib/Linux/libEpicSDL2.so
+  cp build/.libs/libSDL2.so ${TOP_DIR}/Binaries/Linux/libSDL2-2.0.so.0
+  set +x
+}
+
+Run()
+{
+  cd ${TOP_DIR}
+  echo "==> $1"
+  if [ -n "$VERBOSE" ]; then
+	$1
+  else
+	$1 >> ${SCRIPT_DIR}/BuildThirdParty.log 2>&1
+  fi
+}
+
+echo "Building ThirdParty libraries"
+if [ -z "$VERBOSE" ]; then
+  echo "(For full output set VERBOSE=1, otherwise output is logged to BuildThirdParty.log)"
+fi
+
+rm -f ${SCRIPT_DIR}/BuildThirdParty.log
+Run BuildJemalloc
+Run BuildOpus
+Run BuildOgg
+Run BuildVorbis
+Run BuildPNG
+# patches pending - Run BuildHLSLCC
+# patches pending - Run BuildMcpp
+Run BuildFreeType
+Run BuildRecast
+Run BuildICU
+Run BuildSDL2
+
+echo "********** SUCCESS ****************"
diff --git a/Engine/Build/BatchFiles/Linux/GenerateProjectFiles.sh b/Engine/Build/BatchFiles/Linux/GenerateProjectFiles.sh
--- a/Engine/Build/BatchFiles/Linux/GenerateProjectFiles.sh
+++ b/Engine/Build/BatchFiles/Linux/GenerateProjectFiles.sh
@@ -17,14 +17,15 @@ located inside the Engine/Build/BatchFiles/Linux directory."
   exit 1
 fi
 
-
 if [ "$(lsb_release --id)" = "Distributor ID:	Ubuntu" -o "$(lsb_release --id)" = "Distributor ID:	Debian" ]; then
   # Install all necessary dependencies
   DEPS="mono-xbuild \
     mono-dmcs \
     libmono-microsoft-build-tasks-v4.0-4.0-cil \
     libmono-system-data-datasetextensions4.0-cil
-    libmono-system-web-extensions4.0-cil"
+    libmono-system-web-extensions4.0-cil
+    libmono-system-management4.0-cil
+    libogg-dev"
 
   for DEP in $DEPS; do
     if ! dpkg -s $DEP > /dev/null 2>&1; then
@@ -36,6 +37,16 @@ if [ "$(lsb_release --id)" = "Distributor ID:	Ubuntu" -o "$(lsb_release --id)" =
   done
 fi
 
+# Fixes for case sensitive filesystem.
+for BASE in Content/Editor/Slate Content/Slate Documentation/Source/Shared/Icons; do
+  find $BASE -name "*PNG" | while read PNG_UPPER; do
+    png_lower="$(echo "$PNG_UPPER" | sed 's/PNG$/png/')"
+    echo "$PNG_UPPER -> $png_lower"
+    # link, and not move, to make it usable with Perforce workspaces
+    ln -sf `basename "$PNG_UPPER"` "$png_lower"
+  done
+done
+
 set -x
 xbuild Source/Programs/UnrealBuildTool/UnrealBuildTool_Mono.csproj \
   /verbosity:quiet /nologo \
@@ -43,4 +54,4 @@ xbuild Source/Programs/UnrealBuildTool/UnrealBuildTool_Mono.csproj \
   /p:Configuration="Development"
 
 # pass all parameters to UBT
-mono Binaries/DotNET/UnrealBuildTool.exe -makefile "$@"
\ No newline at end of file
+mono Binaries/DotNET/UnrealBuildTool.exe -makefile "$@"
