Runuat.bat: don't use AutomationToolLauncher when running automated tests

diff --git a/Engine/Build/BatchFiles/RunUAT.bat b/Engine/Build/BatchFiles/RunUAT.bat
--- a/Engine/Build/BatchFiles/RunUAT.bat
+++ b/Engine/Build/BatchFiles/RunUAT.bat
@@ -1,4 +1,4 @@
-@echo off
+echo off
 setlocal 
 
 echo Running AutomationTool...
@@ -13,6 +13,14 @@ set UATExecutable=AutomationToolLauncher.exe
 set UATDirectory=Binaries\DotNET\
 set UATNoCompileArg=
 
+rem ## when running automated test, there's no need to use AutomationToolLauncher (fix for EC hangs)
+set UATCmdLine=%*
+set UATCmdLine=%UATCmdLine:EditorTest=%
+set UATCmdLine=%UATCmdLine:GameTest=%
+set UATCmdLine=%UATCmdLine:EditorAutomationTest=%
+set UATCmdLine=%UATCmdLine:GameAutomationTest=%
+if not "%*"=="%UATCmdLine%" set UATExecutable=AutomationTool.exe
+
 rem ## Change the CWD to /Engine. 
 pushd %~dp0..\..\
 if not exist Build\BatchFiles\RunUAT.bat goto Error_BatchFileInWrongLocation
