Runuat.bat: fixing errors when parsing some commandlines.

diff --git a/Engine/Build/BatchFiles/RunUAT.bat b/Engine/Build/BatchFiles/RunUAT.bat
--- a/Engine/Build/BatchFiles/RunUAT.bat
+++ b/Engine/Build/BatchFiles/RunUAT.bat
@@ -15,6 +15,7 @@ set UATNoCompileArg=
 
 rem ## when running automated test, there's no need to use AutomationToolLauncher (fix for EC hangs)
 set UATCmdLine=%*
+set UATCmdLine=%UATCmdLine: =_%
 set UATCmdLine=%UATCmdLine:EditorTest=%
 set UATCmdLine=%UATCmdLine:GameTest=%
 set UATCmdLine=%UATCmdLine:EditorAutomationTest=%
