Fxing Runuat.bat for real

diff --git a/Engine/Build/BatchFiles/RunUAT.bat b/Engine/Build/BatchFiles/RunUAT.bat
--- a/Engine/Build/BatchFiles/RunUAT.bat
+++ b/Engine/Build/BatchFiles/RunUAT.bat
@@ -14,13 +14,15 @@ set UATDirectory=Binaries\DotNET\
 set UATNoCompileArg=
 
 rem ## when running automated test, there's no need to use AutomationToolLauncher (fix for EC hangs)
-set UATCmdLine=%*
+set UATCmdLineEscaped=%*
+set UATCmdLineEscaped=%UATCmdLineEscaped: =_%
+set UATCmdLine=%UATCmdLineEscaped%
 set UATCmdLine=%UATCmdLine: =_%
 set UATCmdLine=%UATCmdLine:EditorTest=%
 set UATCmdLine=%UATCmdLine:GameTest=%
 set UATCmdLine=%UATCmdLine:EditorAutomationTest=%
 set UATCmdLine=%UATCmdLine:GameAutomationTest=%
-if not "%*"=="%UATCmdLine%" set UATExecutable=AutomationTool.exe
+if not "%UATCmdLineEscaped%"=="%UATCmdLine%" set UATExecutable=AutomationTool.exe
 
 rem ## Change the CWD to /Engine. 
 pushd %~dp0..\..\
