Runuat.bat more command line parsing fixes (double quotes)

diff --git a/Engine/Build/BatchFiles/RunUAT.bat b/Engine/Build/BatchFiles/RunUAT.bat
--- a/Engine/Build/BatchFiles/RunUAT.bat
+++ b/Engine/Build/BatchFiles/RunUAT.bat
@@ -1,4 +1,4 @@
-echo off
+rem secho off
 setlocal 
 
 echo Running AutomationTool...
@@ -16,6 +16,7 @@ set UATNoCompileArg=
 rem ## when running automated test, there's no need to use AutomationToolLauncher (fix for EC hangs)
 set UATCmdLineEscaped=%*
 set UATCmdLineEscaped=%UATCmdLineEscaped: =_%
+set UATCmdLineEscaped=%UATCmdLineEscaped:"=_%
 set UATCmdLine=%UATCmdLineEscaped%
 set UATCmdLine=%UATCmdLine: =_%
 set UATCmdLine=%UATCmdLine:EditorTest=%
