Temporarily force use of bundled Mono on Mac if installed version is not 3.2.x, as UAT currently has problems compiling in Mono 3.4.0

diff --git a/Engine/Build/BatchFiles/Mac/SetupMono.sh b/Engine/Build/BatchFiles/Mac/SetupMono.sh
--- a/Engine/Build/BatchFiles/Mac/SetupMono.sh
+++ b/Engine/Build/BatchFiles/Mac/SetupMono.sh
@@ -14,7 +14,7 @@ if [ -f "/usr/bin/mono" ]; then
 	if [ ${MONO_VERSION[0]} -ge 3 ]; then
 		if [ ${MONO_VERSION[1]} -eq 2 ] && [ ${MONO_VERSION[2]} -ge 6 ]; then
 			IS_MONO_INSTALLED=1
-		elif [ ${MONO_VERSION[1]} -gt 2 ]; then
+		elif [ ${MONO_VERSION[1]} -eq 2 ]; then # @todo: change -eq back to -gt when problems with running on Mono 3.4 are solved
 			IS_MONO_INSTALLED=1
 		fi
 	fi
