Common RunUAT.sh for Mac and Linux

[GitHub] part of PR 242.

diff --git a/Engine/Build/BatchFiles/RunUAT.command b/Engine/Build/BatchFiles/RunUAT.command
--- a/Engine/Build/BatchFiles/RunUAT.command
+++ b/Engine/Build/BatchFiles/RunUAT.command
@@ -1,86 +1,10 @@
 #!/bin/sh
-
-## Unreal Engine 4 AutomationTool setup script
 ## Copyright 1998-2014 Epic Games, Inc. All Rights Reserved.
+##
+## Unreal Engine 4 AutomationTool setup script
+##
+## This script uses the .command extenion so that is clickable in
+## in the OSX Finder.  All it does is call RunUAT.sh which does
+## the real work.
 
-## This script is expecting to exist in the UE4/Engine/Build/BatchFiles directory.  It will not work correctly
-## if you copy it to a different location and run it.
-
-echo
-echo Running AutomationTool...
-echo
-
-# Setup Mono
-source "`dirname "$0"`/Mac/SetupMono.sh" "`dirname "$0"`/Mac"
-
-# put ourselves into Engine directory (two up from location of this script)
-pushd "`dirname "$0"`/../.."
-
-UATDirectory=Binaries/DotNET/
-UATNoCompileArg=
-if [ ! -f Build/BatchFiles/RunUAT.command ]; then
-	echo RunUAT ERROR: The batch file does not appear to be located in the /Engine/Build/BatchFiles directory.  This script must be run from within that directory.
-	exit 1
-fi
-
-# see if we have the no compile arg
-echo Checking for -NoCompile
-args=“$@“
-HAS_NO_COMPILE=`echo $args | grep -i nocompile`
-if [ "$HAS_NO_COMPILE" == "$args" ]; then
-	UATNoCompileArg=-NoCompile
-else
-	UATNoCompileArg=
-fi
-
-# see if the .csproj exists to be compiled
-# @todo: Check for an override envvar: if not "%ForcePrecompiledUAT%"=="" goto NoVisualStudio2010Environment
-if [ "$UATNoCompileArg" != "-NoCompile" ]; then
-	if [ ! -f Source/Programs/AutomationTool/AutomationTool_Mono.csproj ]; then
-		echo No project to compile, attempting to use precompiled AutomationTool
-		UATNoCompileArg=-NoCompile
-	else
-		echo Check if $UATDirectory contains AutomationTool.exe
-		if [ -f "$UATDirectory/AutomationTool.exe" ]; then
-			echo AutomationTool exists: Deleting
-			rm -f $UATDirectory/AutomationTool.exe
-			if [ -d "$UATDirectory/AutomationScripts" ]; then
-				echo Deleting all AutomationScript dlls
-				rm -rf $UATDirectory/AutomationScripts
-			fi
-		fi
-		echo Compiling AutomationTool with xbuild
-		xbuild Source/Programs/AutomationTool/AutomationTool_Mono.csproj /property:Configuration=Development /property:Platform=AnyCPU /property:TargetFrameworkVersion=v4.5 /tv:4.0 /verbosity:quiet /nologo |grep -i error
-		# make sure it succeeded
-		if [ $? -ne 1 ]; then
-			echo RunUAT ERROR: AutomationTool failed to compile.
-			exit 1
-		else
-			echo Compilation Succeeded
-		fi
-	fi
-fi
-
-## Run AutomationTool
-
-#run UAT
-pushd $UATDirectory
-if [ -z "$uebp_LogFolder" ]; then
-LogDir=$HOME/Library/Logs/Unreal\ Engine/LocalBuildLogs
-else
-	LogDir="$uebp_LogFolder"
-fi
-# you can't set a dotted env var nicely in sh, but env will run a command with a list of env vars set, including dotted ones
-echo Start UAT: mono AutomationTool.exe "$@"
-env uebp_LogFolder="$LogDir" mono AutomationTool.exe "$@" $UATNoCompileArg
-UATReturn=$?
-
-# @todo: Copy log files to somewhere useful
-# if not "%uebp_LogFolder%" == "" copy log*.txt %uebp_LogFolder%\UAT_*.*
-# if "%uebp_LogFolder%" == "" copy log*.txt c:\LocalBuildLogs\UAT_*.*
-cp log*.txt /var/log
-
-if [ $UATReturn -ne 0 ]; then
-	echo RunUAT ERROR: AutomationTool was unable to run successfully.
-	exit 1
-fi
+exec `dirname "$0"`/RunUAT.sh $@
diff --git a/Engine/Build/BatchFiles/RunUAT.sh b/Engine/Build/BatchFiles/RunUAT.sh
new file mode 100755
--- /dev/null
+++ b/Engine/Build/BatchFiles/RunUAT.sh
@@ -0,0 +1,104 @@
+#!/bin/sh
+## Copyright 1998-2014 Epic Games, Inc. All Rights Reserved.
+##
+## Unreal Engine 4 AutomationTool setup script
+##
+## This script is expecting to exist in the UE4/Engine/Build/BatchFiles directory.  It will not work
+## correctly if you copy it to a different location and run it.
+
+echo
+echo Running AutomationTool...
+echo
+
+# put ourselves into Engine directory (two up from location of this script)
+SCRIPT_DIR=$(cd `dirname "$0"` && pwd)
+cd "$SCRIPT_DIR/../.."
+
+UATDirectory=Binaries/DotNET/
+UATNoCompileArg=
+
+if [ ! -f Build/BatchFiles/RunUAT.sh ]; then
+	echo "RunUAT ERROR: The script does not appear to be located in the "
+  echo "Engine/Build/BatchFiles directory.  This script must be run from within that directory."
+	exit 1
+fi
+
+# see if we have the no compile arg
+args=“$@“
+if echo $args | grep -q -i "\-nocompile"; then
+	UATNoCompileArg=-NoCompile
+else
+	UATNoCompileArg=
+fi
+
+if [ "$(uname)" = "Darwin" ]; then
+	# Fix Mono if needed
+	cd $SCRIPT_DIR/Mac
+	sh FixMonoFiles.sh
+	cd -
+
+	# setup bundled mono
+	echo Setup bundled mono
+	export UE_MONO_DIR=$PWD/Binaries/ThirdParty/Mono/Mac
+	export PATH=$UE_MONO_DIR/bin:$PATH
+	export MONO_PATH=$UE_MONO_DIR/lib:$MONO_PATH
+fi
+
+if [ "$UATNoCompileArg" != "-NoCompile" ]; then
+  # see if the .csproj exists to be compiled
+	if [ ! -f Source/Programs/AutomationTool/AutomationTool_Mono.csproj ]; then
+		echo No project to compile, attempting to use precompiled AutomationTool
+		UATNoCompileArg=-NoCompile
+	else
+		if [ -f "$UATDirectory/AutomationTool.exe" ]; then
+			echo AutomationTool exists: Deleting
+			rm -f $UATDirectory/AutomationTool.exe
+			if [ -d "$UATDirectory/AutomationScripts" ]; then
+				echo Deleting all AutomationScript dlls
+				rm -rf $UATDirectory/AutomationScripts
+			fi
+		fi
+		echo Compiling AutomationTool with xbuild
+
+		ARGS="/p:Configuration=Development /p:Platform=AnyCPU /verbosity:quiet /nologo"
+    ARGS="${ARGS} /p:TargetFrameworkProfile="
+		if [ "$(uname)" = "Darwin" ]; then
+		  ARGS="${ARGS} /p:TargetFrameworkVersion=v4.5"
+		fi
+
+		echo "xbuild Source/Programs/AutomationTool/AutomationTool_Mono.csproj $ARGS"
+		xbuild Source/Programs/AutomationTool/AutomationTool_Mono.csproj $ARGS
+		# make sure it succeeded
+		if [ $? -ne 0 ]; then
+			echo RunUAT ERROR: AutomationTool failed to compile.
+			exit 1
+		else
+			echo Compilation Succeeded
+		fi
+	fi
+fi
+
+## Run AutomationTool
+
+#run UAT
+cd $UATDirectory
+if [ -z "$uebp_LogFolder" ]; then
+  LogDir="$HOME/Library/Logs/Unreal Engine/LocalBuildLogs"
+else
+	LogDir="$uebp_LogFolder"
+fi
+# you can't set a dotted env var nicely in sh, but env will run a command with
+# a list of env vars set, including dotted ones
+echo Start UAT: mono AutomationTool.exe "$@"
+env uebp_LogFolder="$LogDir" mono AutomationTool.exe "$@" $UATNoCompileArg
+UATReturn=$?
+
+# @todo: Copy log files to somewhere useful
+# if not "%uebp_LogFolder%" == "" copy log*.txt %uebp_LogFolder%\UAT_*.*
+# if "%uebp_LogFolder%" == "" copy log*.txt c:\LocalBuildLogs\UAT_*.*
+#cp log*.txt /var/log
+
+if [ $UATReturn -ne 0 ]; then
+	echo RunUAT ERROR: AutomationTool was unable to run successfully.
+	exit 1
+fi
