Fixed RunUAT script on Mac so that it uses SetupMono.sh again

diff --git a/Engine/Build/BatchFiles/RunUAT.sh b/Engine/Build/BatchFiles/RunUAT.sh
--- a/Engine/Build/BatchFiles/RunUAT.sh
+++ b/Engine/Build/BatchFiles/RunUAT.sh
@@ -32,16 +32,8 @@ else
 fi
 
 if [ "$(uname)" = "Darwin" ]; then
-	# Fix Mono if needed
-	cd $SCRIPT_DIR/Mac
-	sh FixMonoFiles.sh
-	cd -
-
-	# setup bundled mono
-	echo Setup bundled mono
-	export UE_MONO_DIR=$PWD/Binaries/ThirdParty/Mono/Mac
-	export PATH=$UE_MONO_DIR/bin:$PATH
-	export MONO_PATH=$UE_MONO_DIR/lib:$MONO_PATH
+	# Setup Mono
+	source "`dirname "$0"`/Mac/SetupMono.sh" "`dirname "$0"`/Mac"
 fi
 
 if [ "$UATNoCompileArg" != "-NoCompile" ]; then
