Fixed problems with running RunUAT from a path with space character

diff --git a/Engine/Build/BatchFiles/RunUAT.command b/Engine/Build/BatchFiles/RunUAT.command
--- a/Engine/Build/BatchFiles/RunUAT.command
+++ b/Engine/Build/BatchFiles/RunUAT.command
@@ -7,4 +7,4 @@
 ## in the OSX Finder.  All it does is call RunUAT.sh which does
 ## the real work.
 
-exec `dirname "$0"`/RunUAT.sh $@
+exec "`dirname "$0"`"/RunUAT.sh $@
diff --git a/Engine/Build/BatchFiles/RunUAT.sh b/Engine/Build/BatchFiles/RunUAT.sh
--- a/Engine/Build/BatchFiles/RunUAT.sh
+++ b/Engine/Build/BatchFiles/RunUAT.sh
@@ -11,7 +11,7 @@ echo Running AutomationTool...
 echo
 
 # put ourselves into Engine directory (two up from location of this script)
-SCRIPT_DIR=$(cd `dirname "$0"` && pwd)
+SCRIPT_DIR=$(cd "`dirname "$0"`" && pwd)
 cd "$SCRIPT_DIR/../.."
 
 UATDirectory=Binaries/DotNET/
