- Fixed RunUAT to quote spaces coming in [ttp 340199]

diff --git a/Engine/Build/BatchFiles/RunUAT.command b/Engine/Build/BatchFiles/RunUAT.command
--- a/Engine/Build/BatchFiles/RunUAT.command
+++ b/Engine/Build/BatchFiles/RunUAT.command
@@ -7,4 +7,4 @@
 ## in the OSX Finder.  All it does is call RunUAT.sh which does
 ## the real work.
 
-exec "`dirname "$0"`"/RunUAT.sh $@
+exec "`dirname "$0"`"/RunUAT.sh "$@"
diff --git a/Engine/Build/BatchFiles/RunUAT.sh b/Engine/Build/BatchFiles/RunUAT.sh
--- a/Engine/Build/BatchFiles/RunUAT.sh
+++ b/Engine/Build/BatchFiles/RunUAT.sh
@@ -1,4 +1,4 @@
-#!/bin/sh
+#!/bin/bash
 ## Copyright 1998-2014 Epic Games, Inc. All Rights Reserved.
 ##
 ## Unreal Engine 4 AutomationTool setup script
@@ -10,6 +10,20 @@ echo
 echo Running AutomationTool...
 echo
 
+# loop over the arguments, quoting spaces to pass to UAT proper
+Args=
+i=0
+for Arg in "$@"
+do
+	# replace all ' ' with '\ '
+	NewArg=${Arg// /\\ }
+	# append it to the array
+	Args[i]=$NewArg
+	# move to next array entry
+	i=$((i+1))
+done
+
+
 # put ourselves into Engine directory (two up from location of this script)
 SCRIPT_DIR=$(cd "`dirname "$0"`" && pwd)
 cd "$SCRIPT_DIR/../.."
@@ -24,8 +38,7 @@ if [ ! -f Build/BatchFiles/RunUAT.sh ]; then
 fi
 
 # see if we have the no compile arg
-args=“$@“
-if echo $args | grep -q -i "\-nocompile"; then
+if echo "${Args[@]}" | grep -q -i "\-nocompile"; then
 	UATNoCompileArg=-NoCompile
 else
 	UATNoCompileArg=
@@ -81,8 +94,8 @@ else
 fi
 # you can't set a dotted env var nicely in sh, but env will run a command with
 # a list of env vars set, including dotted ones
-echo Start UAT: mono AutomationTool.exe "$@"
-env uebp_LogFolder="$LogDir" mono AutomationTool.exe "$@" $UATNoCompileArg
+echo Start UAT: mono AutomationTool.exe "${Args[@]}"
+env uebp_LogFolder="$LogDir" mono AutomationTool.exe $Args $UATNoCompileArg
 UATReturn=$?
 
 # @todo: Copy log files to somewhere useful
