- Fixed the RunUAT script array passing stuff

diff --git a/Engine/Build/BatchFiles/RunUAT.sh b/Engine/Build/BatchFiles/RunUAT.sh
--- a/Engine/Build/BatchFiles/RunUAT.sh
+++ b/Engine/Build/BatchFiles/RunUAT.sh
@@ -95,7 +95,7 @@ fi
 # you can't set a dotted env var nicely in sh, but env will run a command with
 # a list of env vars set, including dotted ones
 echo Start UAT: mono AutomationTool.exe "${Args[@]}"
-env uebp_LogFolder="$LogDir" mono AutomationTool.exe $Args $UATNoCompileArg
+env uebp_LogFolder="$LogDir" mono AutomationTool.exe "${Args[@]}" $UATNoCompileArg
 UATReturn=$?
 
 # @todo: Copy log files to somewhere useful
