- Undid the space fixed with back quotes, caused some problems, and couldn't repro the original issue!

diff --git a/Engine/Build/BatchFiles/RunUAT.sh b/Engine/Build/BatchFiles/RunUAT.sh
--- a/Engine/Build/BatchFiles/RunUAT.sh
+++ b/Engine/Build/BatchFiles/RunUAT.sh
@@ -16,7 +16,9 @@ i=0
 for Arg in "$@"
 do
 	# replace all ' ' with '\ '
-	NewArg=${Arg// /\\ }
+	# DISABLED UNTIL FURTHER INVESTIGATION - IT SEEMS IT WASN'T NEEDED AFTER ALL
+	# NewArg=${Arg// /\\ }
+	NewArg=$Arg
 	# append it to the array
 	Args[i]=$NewArg
 	# move to next array entry
