[Github] PR 242 (safe changes).

- Code cleanup and minor rearrangement (e.g. variable name changes).
- Additions.
- Case issues fixed.
- Minor fixes (init ordering).

diff --git a/Engine/Build/BatchFiles/Linux/BuildThirdParty.sh b/Engine/Build/BatchFiles/Linux/BuildThirdParty.sh
--- a/Engine/Build/BatchFiles/Linux/BuildThirdParty.sh
+++ b/Engine/Build/BatchFiles/Linux/BuildThirdParty.sh
@@ -1,6 +1,6 @@
 #!/bin/bash
 # Script for building ThirdParty modules for linux.
-# You must run InstallDeps.sh to download and patch the sources before running
+# You must run UpdateDeps.sh to download and patch the sources before running
 # this script.
 
 SCRIPT_DIR=$(cd "$(dirname "$BASH_SOURCE")" ; pwd)
@@ -9,12 +9,31 @@ cd ${TOP_DIR}
 set -e
 
 MAKE_ARGS=-j4
+TARGET_ARCH=x86_64-unknown-linux-gnu
 
+# Open files for edit using p4 command line.
+# Does nothing if the file is already writable (which is case for external
+# github developers).
 P4Open()
 {
-  set +x
-  which p4>>/dev/null && p4 open "$@"
-  set -x
+  for file in $@; do
+    # If the file does not exist do nothing
+    if [ ! -e $file ]; then
+      return
+    fi
+    # If the file is already writable do nothing
+    if [ -w $file ]; then
+      return
+    fi
+    # If the file is already writable do nothing
+    if ! which p4 > /dev/null; then
+      echo "File is not writable and 'p4' command not found."
+      exit 1
+    fi
+    set +x
+    p4 open "$file"
+    set -x
+  done
 }
 
 BuildJemalloc()
@@ -26,9 +45,10 @@ BuildJemalloc()
   cd jemalloc-3.4.0
   ./configure --with-mangling --with-jemalloc-prefix=je_340_
   make
-  
-  P4Open ../../lib/Linux/separate/libjemalloc.a
-  cp lib/* ../../lib/Linux/separate/
+  local LIB_DIR=../../lib/Linux/$TARGET_ARCH/separate
+  P4Open $LIB_DIR/libjemalloc.a
+  mkdir -p $LIB_DIR
+  cp lib/* $LIB_DIR/
   cp lib/libjemalloc.so.1 ${TOP_DIR}/Binaries/Linux
   set +x
 }
@@ -40,11 +60,12 @@ BuildOpus()
   cd Source/ThirdParty/libOpus/opus-1.1/
   P4Open configure
   chmod +x configure
-  ./configure
+  ./configure --with-pic
   make $MAKE_ARGS
-
-  cp .libs/libopus.a .libs/libopus*.so* Linux
-  cp Linux/libopus.so.0 ${TOP_DIR}/Binaries/Linux/
+  local LIB_DIR=Linux/$TARGET_ARCH
+  mkdir -p $LIB_DIR
+  cp .libs/libopus.a .libs/libopus*.so* $LIB_DIR
+  cp $LIB_DIR/libopus.so.0 ${TOP_DIR}/Binaries/Linux/
   set +x
 }
 
@@ -54,16 +75,19 @@ BuildPNG()
   set -x
   cd Source/ThirdParty/libPNG/libPNG-1.5.2
   cp --remove-destination scripts/makefile.linux Makefile
+  cp --remove-destination scripts/pnglibconf.h.prebuilt pnglibconf.h
   P4Open pnglibconf.h
   make $MAKE_ARGS
 
-  P4Open lib/Linux/libpng.a
-  cp libpng.a lib/Linux
-  cp libpng15.so lib/Linux/libpng.so
-  if [ ! -d ../../../../Binaries/Linux ]; then
-    mkdir ../../../../Binaries/Linux
-  fi
-  cp -a libpng15.so* ${TOP_DIR}/Binaries/Linux/
+  local LIB_DIR=lib/Linux/$TARGET_ARCH
+  mkdir -p $LIB_DIR
+  P4Open $LIB_DIR/libpng.a
+  cp --remove-destination libpng.a $LIB_DIR/
+  # (amigo) this dylib will mess with our build
+  # instead we want to use system's libpng to avoid asserts
+  #cp libpng15.so $LIB_DIR/libpng.so
+  mkdir -p ${TOP_DIR}/Binaries/Linux/
+  #cp -a libpng15.so* ${TOP_DIR}/Binaries/Linux/
   set +x
 }
 
@@ -74,11 +98,13 @@ BuildOgg()
   cd Source/ThirdParty/Ogg/libogg-1.2.2/
   P4Open configure
   chmod +x configure
-  ./configure
+  ./configure --with-pic
   make $MAKE_ARGS
 
-  P4Open lib/Linux/libogg.a
-  cp src/.libs/* lib/Linux
+  local LIB_DIR=lib/Linux/$TARGET_ARCH
+  P4Open $LIB_DIR/libogg.a
+  mkdir -p $LIB_DIR
+  cp src/.libs/* $LIB_DIR
   set +x
 }
 
@@ -89,11 +115,15 @@ BuildVorbis()
   cd Source/ThirdParty/Vorbis/libvorbis-1.3.2/
   P4Open configure
   chmod +x configure
-  OGG_LIBS=../../Ogg/libogg-1.2.2/lib/Linux OGG_CFLAGS=-I../../../Ogg/libogg-1.2.2/include ./configure --with-pic --disable-oggtest
+  OGG_LIBS=../../Ogg/libogg-1.2.2/lib/Linux \
+      OGG_CFLAGS=-I../../../Ogg/libogg-1.2.2/include \
+      ./configure --with-pic --disable-oggtest
   make $MAKE_ARGS
 
-  P4Open lib/Linux/libvorbis.a lib/Linux/libvorbisfile.a
-  cp lib/.libs/* lib/Linux
+  local LIB_DIR=lib/Linux/$TARGET_ARCH
+  P4Open $LIB_DIR/libvorbis.a $LIB_DIR/libvorbisfile.a
+  mkdir -p $LIB_DIR
+  cp lib/.libs/libvorbis*.a $LIB_DIR/
   set +x
   cd -
 }
@@ -118,10 +148,12 @@ BuildMcpp()
   ./configure --enable-shared --enable-shared --enable-mcpplib
   make $MAKE_ARGS
 
-  P4Open lib/Linux/libmcpp.a
-  P4Open lib/Linux/libmcpp.so
-  cp --remove-destination ./src/.libs/libmcpp.a lib/Linux
-  cp --remove-destination ./src/.libs/libmcpp.so lib/Linux
+  local LIB_DIR=lib/Linux/$TARGET_ARCH
+  P4Open $LIB_DIR/libmcpp.a
+  P4Open $LIB_DIR/libmcpp.so
+  mkdir -p $LIB_DIR
+  cp --remove-destination ./src/.libs/libmcpp.a $LIB_DIR/
+  cp --remove-destination ./src/.libs/libmcpp.so $LIB_DIR/
   cp --remove-destination ./src/.libs/libmcpp.so ${TOP_DIR}/Binaries/Linux/libmcpp.so.0
   set +x
 }
@@ -134,19 +166,10 @@ BuildFreeType()
   pwd
   make $MAKE_ARGS -f ../Builds/Linux/makefile $*
 
-  P4Open ../Lib/Linux/libfreetype2412.a
-  cp --remove-destination libfreetype* ../Lib/Linux
-  set +x
-}
-
-BuildRecast()
-{
-  echo "building Recast"
-  set -x
-  cd Source/ThirdParty/Recast
-  pwd
-  P4Open Epic/lib/Linux/librecast.a
-  make $MAKE_ARGS -f Epic/build/Linux/makefile $*
+  local LIB_DIR=../Lib/Linux/$TARGET_ARCH
+  P4Open $LIB_DIR/libfreetype2412.a
+  mkdir -p $LIB_DIR
+  cp --remove-destination libfreetype* $LIB_DIR/
   set +x
 }
 
@@ -154,56 +177,177 @@ BuildICU()
 {
   echo "building libICU"
   set -x
-  cd Source/ThirdParty/ICU/icu4c-51_2/source
+  cd Source/ThirdParty/ICU/icu4c-53_1/source
   P4Open configure
   chmod +x configure
-  CPPFLAGS=-fPIC ./configure --enable-static --with-data-packaging=static
+  CPPFLAGS=-fPIC ./configure --enable-static
   pwd
   make $MAKE_ARGS $*
 
-  P4Open ../Linux/libicudata.a ../Linux/libicudatad.a ../Linux/libicui18n.a ../Linux/libicui18nd.a
-  P4Open ../Linux/libicuio.a ../Linux/libicuiod.a ../Linux/libicule.a ../Linux/libiculed.a
-  P4Open ../Linux/libiculx.a ../Linux/libiculxd.a ../Linux/libicutu.a ../Linux/libicutud.a
-  P4Open ../Linux/libicuuc.a ../Linux/libicuucd.a
-  cp --remove-destination lib/*.a ../Linux
+  local LIB_DIR=../Linux/$TARGET_ARCH
+  mkdir -p $LIB_DIR
+  mkdir -p ${TOP_DIR}/Binaries/ThirdParty/ICU/icu4c-53_1/Linux/$TARGET_ARCH/
+  P4Open $LIB_DIR/libicudata.a $LIB_DIR/libicudatad.a $LIB_DIR/libicui18n.a $LIB_DIR/libicui18nd.a
+  P4Open $LIB_DIR/libicuio.a $LIB_DIR/libicuiod.a $LIB_DIR/libicule.a $LIB_DIR/libiculed.a
+  P4Open $LIB_DIR/libiculx.a $LIB_DIR/libiculxd.a $LIB_DIR/libicutu.a $LIB_DIR/libicutud.a
+  P4Open $LIB_DIR/libicuuc.a $LIB_DIR/libicuucd.a
+  cp --remove-destination lib/*.a $LIB_DIR
+  cp -P --remove-destination lib/*.so* ${TOP_DIR}/Binaries/ThirdParty/ICU/icu4c-53_1/Linux/$TARGET_ARCH/
+
   set +x
 }
 
-BuildSDL2()
+BuildLND()
 {
-  echo "building dsext"
+  echo "building LinuxNativeDialogs"
   set -x
-  cd Source/ThirdParty/SDL2/DisplayServerExtensions
-  P4Open bin/libdsext.a
+  cd Source/ThirdParty/LinuxNativeDialogs/UELinuxNativeDialogs/build
+  cmake ..
   make $MAKE_ARGS
+  local LIB_DIR=../lib/Linux/$TARGET_ARCH/
+  if [ -f libqt4dialog.so ]; then
+    ln -s libqt4dialog.so libLND.so
+    mv *.so $LIB_DIR/
+    cp -P --remove-destination $LIB_DIR/*.so ${TOP_DIR}/Binaries/Linux/
+  fi
+  if [ -f libqt5dialog.so ]; then
+    ln -s libqt5dialog.so libLND.so
+    mv *.so $LIB_DIR/
+    cp -P --remove-destination $LIB_DIR/*.so ${TOP_DIR}/Binaries/Linux/
+  fi
+  if [ -f libgtk2dialog.so ]; then
+    ln -s libgtk2dialog.so libLND.so
+    mv *.so $LIB_DIR/
+    cp -P --remove-destination $LIB_DIR/*.so ${TOP_DIR}/Binaries/Linux/
+  fi
+  if [ -f libgtk3dialog.so ]; then
+    ln -s libgtk3dialog.so libLND.so
+    mv *.so $LIB_DIR/
+    cp -P --remove-destination $LIB_DIR/*.so ${TOP_DIR}/Binaries/Linux/
+  fi
+  set +x
+}
+
+BuildForsythTriOO()
+{
+  echo "building ForsythTriOO"
+  set -x
+  cd Source/ThirdParty/ForsythTriOO
+  mkdir -p Build/Linux
+  mkdir -p Lib/Linux/$TARGET_ARCH
+  cd Build/Linux
+  cmake -DCMAKE_BUILD_TYPE=Debug ../../
+  make
+  cmake -DCMAKE_BUILD_TYPE=Release ../../
+  make
+  cp --remove-destination libForsythTriOptimizer* ../../Lib/Linux/$TARGET_ARCH/
+  set +x
+}
+
+BuildnvTriStrip()
+{
+  echo "building nvTriStrip"
+  set -x
+  cd Source/ThirdParty/nvTriStrip/nvTriStrip-1.0.0
+  mkdir -p Build/Linux
+  mkdir -p Lib/Linux/$TARGET_ARCH
+  cd Build/Linux
+  cmake -DCMAKE_BUILD_TYPE=Debug ../../
+  make
+  cmake -DCMAKE_BUILD_TYPE=Release ../../
+  make
+  cp --remove-destination libnvtristrip* ../../Lib/Linux/$TARGET_ARCH/
   set +x
+}
+
+BuildnvTextureTools()
+{
+  echo "building nvTextureTools"
+  set -x
+  cd Source/ThirdParty/nvTextureTools/nvTextureTools-2.0.8
+  mkdir -p lib/Linux/$TARGET_ARCH
+  cd src
+  ./configure
+  make
+
+  local LIB_DIR=../lib/Linux/$TARGET_ARCH
+  cp --remove-destination build-debug/src/nvcore/libnvcore.so $LIB_DIR/
+  cp --remove-destination build-debug/src/nvimage/libnvimage.so $LIB_DIR/
+  cp --remove-destination build-debug/src/nvmath/libnvmath.so $LIB_DIR/
+  cp --remove-destination build-debug/src/nvtt/libnvtt.so $LIB_DIR/
+  cp --remove-destination build-debug/src/nvthread/libnvthread.a $LIB_DIR/
+  cp --remove-destination build-debug/src/nvtt/squish/libsquish.a $LIB_DIR/
+  cp --remove-destination $LIB_DIR/*.so ${TOP_DIR}/Binaries/Linux/
+  set +x
+}
+
+BuildSDL2()
+{
   echo "building SDL2"
   set -x
-  cd ../SDL2-2.0.3/build
-  pwd
-  tar xf SDL2-2.0.3.tar.gz
-  cd SDL2-2.0.3
+  cd Source/ThirdParty/SDL2/SDL2-2.0.3/build/SDL2-2.0.3
   P4Open configure
   chmod +x configure
+  chmod +x autogen.sh
+  chmod +x Android.mk
   ./configure
   make $MAKE_ARGS
-
-  # FIXME: build libdsext?
-  P4Open ../../lib/Linux/libEpicSDL2.a  
-  cp --remove-destination build/.libs/libSDL2.a ../../lib/Linux/libEpicSDL2.a
-  cp build/.libs/libSDL2.so ../../lib/Linux/libEpicSDL2.so
+  cp --remove-destination include/* ../../include/
+  local LIB_DIR=../../lib/Linux/$TARGET_ARCH
+  P4Open $LIB_DIR/libEpicSDL2.a
+  mkdir -p $LIB_DIR
+  cp --remove-destination build/.libs/libSDL2.a $LIB_DIR/libEpicSDL2.a
+  cp build/.libs/libSDL2.so $LIB_DIR/libEpicSDL2.so
   cp build/.libs/libSDL2.so ${TOP_DIR}/Binaries/Linux/libSDL2-2.0.so.0
   set +x
 }
 
+Buildcoremod()
+{
+  echo "building coremod"
+  set -x
+  cd Source/ThirdParty/coremod/coremod-4.2.6
+  pwd
+  autoconf
+  CFLAGS=-fPIC ./configure --enable-static
+  make $MAKE_ARGS
+  local LIB_DIR=lib/Linux/$TARGET_ARCH
+  mkdir -p $LIB_DIR
+  cp lib/libxmp-coremod.a $LIB_DIR/libcoremodLinux.a
+  set +x
+}
+
+FixICU()
+{
+  echo "setting symlinks for ICU libraries (possibly temporary)"
+  set -x
+  cd Binaries/ThirdParty/ICU/icu4c-53_1/Linux/x86_64-unknown-linux-gnu
+  pwd
+  ln -sf libicudata.so libicudata.so.53
+  ln -sf libicudata.so libicudata.53.1.so
+  ln -sf libicui18n.so libicui18n.so.53
+  ln -sf libicui18n.so libicui18n.53.1.so
+  ln -sf libicuio.so libicuio.so.53
+  ln -sf libicuio.so libicuio.53.1.so
+  ln -sf libicule.so libicule.so.53
+  ln -sf libicule.so libicule.53.1.so
+  ln -sf libiculx.so libiculx.so.53
+  ln -sf libiculx.so libiculx.53.1.so
+  ln -sf libicutu.so libicutu.so.53
+  ln -sf libicutu.so libicutu.53.1.so
+  ln -sf libicuuc.so libicuuc.so.53
+  ln -sf libicuuc.so libicuuc.53.1.so
+  set +x
+}
+
 Run()
 {
   cd ${TOP_DIR}
   echo "==> $1"
   if [ -n "$VERBOSE" ]; then
-	$1
+    $1
   else
-	$1 >> ${SCRIPT_DIR}/BuildThirdParty.log 2>&1
+    $1 >> ${SCRIPT_DIR}/BuildThirdParty.log 2>&1
   fi
 }
 
@@ -213,16 +357,24 @@ if [ -z "$VERBOSE" ]; then
 fi
 
 rm -f ${SCRIPT_DIR}/BuildThirdParty.log
-Run BuildJemalloc
-Run BuildOpus
-Run BuildOgg
-Run BuildVorbis
-Run BuildPNG
-Run BuildHLSLCC
-Run BuildMcpp
-Run BuildFreeType
-Run BuildRecast
-Run BuildICU
-Run BuildSDL2
+if [ -z "$1" ]; then
+  Run BuildJemalloc
+  Run BuildOpus
+  Run BuildOgg
+  Run BuildVorbis
+  Run BuildPNG
+  Run BuildHLSLCC
+  Run BuildMcpp
+  Run BuildFreeType
+  Run BuildLND
+  Run BuildForsythTriOO
+  Run BuildnvTriStrip
+  Run BuildnvTextureTools
+  Run BuildSDL2
+  Run Buildcoremod
+  Run FixICU
+else
+  Run $1
+fi
 
 echo "********** SUCCESS ****************"
diff --git a/Engine/Build/BatchFiles/Linux/GenerateProjectFiles.sh b/Engine/Build/BatchFiles/Linux/GenerateProjectFiles.sh
--- a/Engine/Build/BatchFiles/Linux/GenerateProjectFiles.sh
+++ b/Engine/Build/BatchFiles/Linux/GenerateProjectFiles.sh
@@ -25,6 +25,7 @@ if [ "$(lsb_release --id)" = "Distributor ID:	Ubuntu" -o "$(lsb_release --id)" =
     libmono-system-data-datasetextensions4.0-cil
     libmono-system-web-extensions4.0-cil
     libmono-system-management4.0-cil
+    libmono-system-xml-linq4.0-cil
     libogg-dev"
 
   for DEP in $DEPS; do
@@ -39,11 +40,14 @@ fi
 
 # Fixes for case sensitive filesystem.
 for BASE in Content/Editor/Slate Content/Slate Documentation/Source/Shared/Icons; do
-  find $BASE -name "*PNG" | while read PNG_UPPER; do
-    png_lower="$(echo "$PNG_UPPER" | sed 's/PNG$/png/')"
-    echo "$PNG_UPPER -> $png_lower"
-    # link, and not move, to make it usable with Perforce workspaces
-    ln -sf `basename "$PNG_UPPER"` "$png_lower"
+  find $BASE -name "*.PNG" | while read PNG_UPPER; do
+    png_lower="$(echo "$PNG_UPPER" | sed 's/.PNG$/.png/')"
+    if [ ! -f $png_lower ]; then
+      PNG_UPPER=$(basename $PNG_UPPER)
+      echo "$png_lower -> $PNG_UPPER"
+      # link, and not move, to make it usable with Perforce workspaces
+      ln -sf `basename "$PNG_UPPER"` "$png_lower"
+    fi
   done
 done
 
diff --git a/Engine/Build/BatchFiles/Linux/README.md b/Engine/Build/BatchFiles/Linux/README.md
--- a/Engine/Build/BatchFiles/Linux/README.md
+++ b/Engine/Build/BatchFiles/Linux/README.md
@@ -6,7 +6,7 @@ So far this has only been tested on Debian/Sid and Ubuntu/Precise.
 
 A wiki page page regarding building on linux is also maintained at
 https://wiki.unrealengine.com/Building_On_Linux, and public IRC discussion
-happens on #UE4Linux on freenode.
+happens on #UE4Linux on FreeNode.
 
 Prerequisites
 -------------
@@ -23,9 +23,10 @@ EU4. The command to install these on Ubundu 12.04 (Precise) is:
 Building
 --------
 
-1. Download the additional .zip archive from the github releases page.
+1. Download the additional .zip archives for the master branch from Epic's GitHub
+   Releases page at https://github.com/EpicGames/UnrealEngine/releases
 
-2. Download v129 of the Steamworks SDK.
+2. Download v129a of the Steamworks SDK.
 
 3. Install the downloaded dependencies with:
 
diff --git a/Engine/Build/BatchFiles/Linux/UpdateDeps.sh b/Engine/Build/BatchFiles/Linux/UpdateDeps.sh
--- a/Engine/Build/BatchFiles/Linux/UpdateDeps.sh
+++ b/Engine/Build/BatchFiles/Linux/UpdateDeps.sh
@@ -1,5 +1,5 @@
 #!/bin/bash
-# This script will verify the downloaded zip file depenencies, unzip them and
+# This script will verify the downloaded zip file dependencies, unzip them and
 # apply linux-specific patches.  It tries to clean up any previous versions
 # using 'git clean' on certain directories.
 # WARNING: This script can be destructive so use with caution if you have
@@ -31,21 +31,41 @@ cd ${TOP_DIR}
 set -e
 trap "echo '==> An error occured while running UpdateDeps.sh'" ERR
 
-FILES="Optional.zip Required_1of3.zip Required_2of3.zip Required_3of3.zip"
-REQUIRED_SHA1SUMS=$SCRIPT_DIR/dependencies.sha1sums
-INSTALLED_SHA1SUMS=$REQUIRED_SHA1SUMS.installed
+DEPENDENCIES=$SCRIPT_DIR/dependencies.txt
+INSTALLED_DEPENDENCIES=${DEPENDENCIES%.txt}.installed
 
+#TODO: Move these hardcodings out of this file
 PERFORCE_ARCHIVE=$ARCHIVE_ROOT/p4api.tgz
-STEAM_ARCHIVE=$ARCHIVE_ROOT/steamworks_sdk_129.zip
-FBX_ARCHIVE=$ARCHIVE_ROOT/fbx20142_1_fbxsdk_linux.tar.gz
+STEAM_ARCHIVE=$ARCHIVE_ROOT/steamworks_sdk_129a.zip
+
+# Make it super easy for users of certain linux distributions.
+CheckDistroDependencies() {
+  echo "==> Checking build dependencies"
+  if [ "$(lsb_release --id)" = "Distributor ID:	Ubuntu" -o "$(lsb_release --id)" = "Distributor ID:	Debian" ]; then
+    # Install all necessary dependencies
+    DEPS="libqt4-dev
+      libglew-dev
+      dos2unix 
+      clang-3.3"
+
+    for DEP in $DEPS; do
+      if ! dpkg -s $DEP > /dev/null 2>&1; then
+        echo "Attempting installation of missing package: $DEP"
+        set -x
+        sudo apt-get install $DEP
+        set +x
+      fi
+    done
+  fi
+}
 
 CheckInstalled() {
   echo "==> Checking existing install"
   # Phase one, check sha1 sums of previously installed files.
-  if [ -f $INSTALLED_SHA1SUMS ]; then
-    if [ "`cat $INSTALLED_SHA1SUMS`" = "`cat $REQUIRED_SHA1SUMS`" ]; then
+  if [ -f $INSTALLED_DEPENDENCIES ]; then
+    if [ "`cat $INSTALLED_DEPENDENCIES`" = "`cat $DEPENDENCIES`" ]; then
       echo "Dependencies are up-to-date."
-      echo "(remove $INSTALLED_SHA1SUMS file and re-run to force install)."
+      echo "(remove $INSTALLED_DEPENDENCIES file and re-run to force install)."
       exit 0
     else
       echo "Installed dependencies are out-of-date."
@@ -56,35 +76,39 @@ CheckInstalled() {
 }
 
 Download() {
-  # For those archives that are missing some can be downloaded automatically using curl
-  # The steamworks SDK lives on a standard public URL but seems to be authorised based
-  # on source IP for a certain amount of time, afterwhich we get a 401 error.
-  #if [ ! -f $STEAM_ARCHIVE ]; then
-  #  echo "==> Downloading steamworks SDK"
-  #  curl --progress -o $STEAM_ARCHIVE https://partner.steamgames.com/downloads/steamworks_sdk_129.zip
-  #fi
-  if [ ! -f $PERFORCE_ARCHIVE ]; then
-	echo "==> Downloading p4api.tgz"
-    curl --progress -o $PERFORCE_ARCHIVE http://ftp.perforce.com/perforce/r12.1/bin.linux26x86_64/p4api.tgz
-  fi
-  if [ ! -f $FBX_ARCHIVE ]; then
-	echo "==> Downloading FBX SDK"
-	curl --progress -o $FBX_ARCHIVE http://images.autodesk.com/adsk/files/fbx20142_1_fbxsdk_linux.tar.gz
-  fi
+  # For those archives that are missing some can be downloaded automatically
+  # using curl.
+  cat $DEPENDENCIES | while read line; do
+    DEP_URL=`echo $line | cut -d " " -f 2`
+    DEP_WALLED=`echo $line | cut -d " " -f 3`
+    DEP_FILE=$(basename $DEP_URL)
+    if [ ! -f $ARCHIVE_ROOT/$DEP_FILE ]; then
+      if [ "$DEP_WALLED" -eq 0 ]; then
+        echo "==> Downloading $DEP_FILE"
+        curl --progress -o $ARCHIVE_ROOT/$DEP_FILE $DEP_URL
+      fi
+    fi
+  done
 }
 
 VerifyDownloads() {
   echo "==> Verifying downloaded archives"
   cd $ARCHIVE_ROOT
-  if ! sha1sum -c $REQUIRED_SHA1SUMS; then
-	echo ""
-    echo "ERROR: One or more archives files are missing or out-of-date."
-	echo "Please download zip files from github release page and place them in "
-	echo "'$ARCHIVE_ROOT'."
-	echo "If you have already downloaded them somewhere else then you can pass "
-	echo "the location as the first argument to this script."
-    exit 1
-  fi
+  cat $DEPENDENCIES | while read line; do
+    DEP_HASH=`echo $line | cut -d " " -f 1`
+    DEP_URL=`echo $line | cut -d " " -f 2`
+    DEP_FILE=$(basename $DEP_URL)
+    if ! echo "$DEP_HASH  $DEP_FILE" | sha1sum -c -; then
+      echo ""
+      echo "ERROR: Missing or out-of-date dependency: $DEP_FILE"
+      echo "Please download the following file to $ARCHIVE_ROOT:"
+      echo "  $DEP_URL"
+      echo "If you have downloaded these files somewhere else you can"
+      echo "pass the location as the first argument to this script."
+      exit 1
+    fi
+  done
+
   cd ${TOP_DIR}
 }
 
@@ -107,6 +131,7 @@ ExtractArchives() {
     Engine/Binaries
     Engine/Source/ThirdParty
     Engine/Plugins/Script/ScriptGeneratorPlugin/Binaries
+    Engine/Plugins/Messaging/UdpMessaging/Binaries
     Engine/Content
     Engine/Documentation
     Engine/Extras
@@ -118,14 +143,19 @@ ExtractArchives() {
 
   echo "==> Installing dependencies"
   # Extract files and write new sha1 files.
-  for file in $FILES; do
-    file_full=$ARCHIVE_ROOT/$file
+  cat $DEPENDENCIES | while read line; do
+    DEP_URL=`echo $line | cut -d " " -f 2`
+    if [[ $(dirname $DEP_URL) != https://github.com/EpicGames/* ]]; then
+      continue
+    fi
+    DEP_FILE=$(basename $DEP_URL)
+    file_full=$ARCHIVE_ROOT/$DEP_FILE
     if [ ! -f $file_full ]; then
       echo "Zip file dependency not found: $file_full"
       exit 1
     fi
     echo "Extracting $file_full"
-    unzip -q $file_full
+    unzip -q -n $file_full
   done
 }
 
@@ -182,16 +212,30 @@ ApplyPatches() {
   chmod +x config.guess
   cd -
 
-  cd ICU/icu4c-51_2/source/
-  chmod +x configure
-  chmod +x config.sub
-  chmod +x config.guess
+  cd SDL2/
+  pwd
+  tar xf SDL2-2.0.3/build/SDL2-2.0.3.tar.gz -C SDL2-2.0.3/build
+  rm -f DisplayServerExtensions/bin/libdsext.a
+  patch -f -p1 < ${SCRIPT_DIR}/SDL2.patch
   cd -
 
-  patch -p1 < ${SCRIPT_DIR}/hlslcc.patch
-  patch -p1 < ${SCRIPT_DIR}/nvtesslib.patch
-  patch -p1 < ${SCRIPT_DIR}/HACD.patch
-  patch -p1 < ${SCRIPT_DIR}/mcpp.patch
+  # This needs to be moved out of the way as the LinuxToolChain.cs
+  # Links via -l and not via path name so will pick the .so instead
+  cd FBX/2014.2.1/
+  pwd
+  rm -f lib/linux/x86_64-unknown-linux-gnu/libfbxsdk.so
+  cd -
+
+  cd ${TOP_DIR}
+}
+
+RenameContent() {
+  echo "==> Correcting Filename case"
+  cd Engine/Content
+  echo "PWD=$PWD"
+
+  find . -name "*.PNG" -exec rename 's/\.PNG$/\.png/' {} \;
+
   cd ${TOP_DIR}
 }
 
@@ -206,42 +250,38 @@ InstallPerforce() {
 
 InstallSteamworks() {
   echo "==> Installing Steamworks SDK"
-  rm -rf Engine/Source/ThirdParty/Steamworks/Steamv129
-  mkdir Engine/Source/ThirdParty/Steamworks/Steamv129
-  cd Engine/Source/ThirdParty/Steamworks/Steamv129
+  rm -rf Engine/Source/ThirdParty/Steamworks/Steamv129a
+  mkdir Engine/Source/ThirdParty/Steamworks/Steamv129a
+  cd Engine/Source/ThirdParty/Steamworks/Steamv129a
   unzip -q $STEAM_ARCHIVE
   cd ${TOP_DIR}
   mkdir -p Engine/Binaries/Linux/
   rm -f Engine/Binaries/Linux/libsteam_api.so
-  ln -s Engine/Source/ThirdParty/Steamworks/Steamv129/sdk/redistributable_bin/linux64/libsteam_api.so Engine/Binaries/Linux/
-}
-
-InstallFBX() {
-  echo "==> Installing FBX"
-  cd Engine/Source/ThirdParty/FBX
-  tar xf $FBX_ARCHIVE
-  rm -rf fbx
-  mkdir fbx
-  yes yes | ./fbx20142_1_fbxsdk_linux fbx > /dev/null
-  rm -rf 2014.2.1/lib/linux
-  mkdir 2014.2.1/lib/linux
-  cp fbx/lib/gcc4/x64/release/libfbxsdk.a 2014.2.1/lib/linux
+  cd Engine/Binaries/Linux
+  ln -s ../../Source/ThirdParty/Steamworks/Steamv129a/sdk/redistributable_bin/linux64/libsteam_api.so
   cd ${TOP_DIR}
 }
 
 UpdateInstallStamp() {
-  echo "==> Updating install stamp ($INSTALLED_SHA1SUMS)"
-  cp $REQUIRED_SHA1SUMS $INSTALLED_SHA1SUMS
+  echo "==> Updating install stamp ($INSTALLED_DEPENDENCIES)"
+  cp $DEPENDENCIES $INSTALLED_DEPENDENCIES
 }
 
+CheckDistroDependencies
 CheckInstalled
 Download
 VerifyDownloads
 ExtractArchives
 ApplyPatches
-InstallFBX
+RenameContent
 InstallPerforce
 InstallSteamworks
 UpdateInstallStamp
 
 echo "********** SUCCESS ****************"
+
+# Once we've successfully updated the deps we need to rebuild the
+# ThirdParty modules.  This is mostly so that we get everything built
+# with -fPIC.  BuildThirdParty.sh is a seperate script so that it can
+# be run repeatedly and independently of UpdateDeps.sh if needed.
+# $SCRIPT_DIR/BuildThirdParty.sh
diff --git a/Engine/Build/BatchFiles/Linux/dependencies.txt b/Engine/Build/BatchFiles/Linux/dependencies.txt
new file mode 100644
--- /dev/null
+++ b/Engine/Build/BatchFiles/Linux/dependencies.txt
@@ -0,0 +1,5 @@
+5a8490d5de3681d9579127866e25696ff0d2e7ac https://github.com/EpicGames/UnrealEngine/releases/download/latest-preview/Optional.zip  1
+f5458c00a15ccdf3ca42d4cc3f8d4ff938299537 https://github.com/EpicGames/UnrealEngine/releases/download/latest-preview/Required_1of2.zip  1
+f334c28795760d4c235f49a1436bb523c1ac539a https://github.com/EpicGames/UnrealEngine/releases/download/latest-preview/Required_2of2.zip  1
+7a3d986efda677060a616a23e12c09445d3e0dca http://ftp.perforce.com/perforce/r12.1/bin.linux26x86_64/p4api.tgz  0
+4f0fdaf5ce8ee2bfbe4dfd7e58ffb25963c840f1 https://partner.steamgames.com/downloads/steamworks_sdk_129a.zip  1
diff --git a/Engine/Config/BaseDeviceProfiles.ini b/Engine/Config/BaseDeviceProfiles.ini
--- a/Engine/Config/BaseDeviceProfiles.ini
+++ b/Engine/Config/BaseDeviceProfiles.ini
@@ -23,6 +23,9 @@
 +DeviceProfileNameAndTypes=MacNoEditor,MacNoEditor
 +DeviceProfileNameAndTypes=MacServer,MacServer
 +DeviceProfileNameAndTypes=WinRT,WinRT
++DeviceProfileNameAndTypes=Linux,Linux
++DeviceProfileNameAndTypes=LinuxNoEditor,LinuxNoEditor
++DeviceProfileNameAndTypes=LinuxServer,LinuxServer
 
 [Windows DeviceProfile]
 DeviceType=Windows
@@ -191,3 +194,21 @@ BaseProfileName=Mac
 [WinRT DeviceProfile]
 DeviceType=WinRT
 BaseProfileName=
+
+[Linux DeviceProfile]
+DeviceType=Linux
+BaseProfileName=
+MeshLODSettings=
+TextureLODSettings=
+
+[LinuxNoEditor DeviceProfile]
+DeviceType=LinuxNoEditor
+BaseProfileName=Linux
+MeshLODSettings=
+TextureLODSettings=
+
+[LinuxServer DeviceProfile]
+DeviceType=LinuxServer
+BaseProfileName=Linux
+MeshLODSettings=
+TextureLODSettings=
diff --git a/Engine/Source/Developer/AudioFormatOgg/AudioFormatOgg.Build.cs b/Engine/Source/Developer/AudioFormatOgg/AudioFormatOgg.Build.cs
--- a/Engine/Source/Developer/AudioFormatOgg/AudioFormatOgg.Build.cs
+++ b/Engine/Source/Developer/AudioFormatOgg/AudioFormatOgg.Build.cs
@@ -18,6 +18,7 @@ public class AudioFormatOgg : ModuleRules
 		if ((Target.Platform == UnrealTargetPlatform.Win64) ||
 			(Target.Platform == UnrealTargetPlatform.Win32) ||
 			(Target.Platform == UnrealTargetPlatform.Mac) ||
+			(Target.Platform == UnrealTargetPlatform.Linux) ||
             ((Target.Platform == UnrealTargetPlatform.HTML5 && Target.Architecture == "-win32"))
             )
 		{
diff --git a/Engine/Source/Developer/AudioFormatOpus/AudioFormatOpus.Build.cs b/Engine/Source/Developer/AudioFormatOpus/AudioFormatOpus.Build.cs
--- a/Engine/Source/Developer/AudioFormatOpus/AudioFormatOpus.Build.cs
+++ b/Engine/Source/Developer/AudioFormatOpus/AudioFormatOpus.Build.cs
@@ -16,7 +16,8 @@ public class AudioFormatOpus : ModuleRules
 			);
 
 		if ((Target.Platform == UnrealTargetPlatform.Win64) ||
-			(Target.Platform == UnrealTargetPlatform.Win32) //||
+			(Target.Platform == UnrealTargetPlatform.Win32) ||
+            (Target.Platform == UnrealTargetPlatform.Linux)
 			//(Target.Platform == UnrealTargetPlatform.Mac) ||
             //((Target.Platform == UnrealTargetPlatform.HTML5 && Target.Architecture == "-win32"))
             )
diff --git a/Engine/Source/Developer/DesktopPlatform/DesktopPlatform.Build.cs b/Engine/Source/Developer/DesktopPlatform/DesktopPlatform.Build.cs
--- a/Engine/Source/Developer/DesktopPlatform/DesktopPlatform.Build.cs
+++ b/Engine/Source/Developer/DesktopPlatform/DesktopPlatform.Build.cs
@@ -9,5 +9,10 @@ public class DesktopPlatform : ModuleRules
 		PrivateIncludePaths.Add("Developer/DesktopPlatform/Private");
 
 		PrivateDependencyModuleNames.Add("Core");
+
+		if (Target.Platform == UnrealTargetPlatform.Linux)
+		{
+            AddThirdPartyPrivateStaticDependencies(Target, "SDL2");
+		}
 	}
 }
diff --git a/Engine/Source/Developer/LauncherServices/Private/Launcher/LauncherBuildCommands.h b/Engine/Source/Developer/LauncherServices/Private/Launcher/LauncherBuildCommands.h
--- a/Engine/Source/Developer/LauncherServices/Private/Launcher/LauncherBuildCommands.h
+++ b/Engine/Source/Developer/LauncherServices/Private/Launcher/LauncherBuildCommands.h
@@ -68,7 +68,7 @@ public:
 		FString CommandLine;
 
 		FString Platform = TEXT("Win64");
-		if (TargetPlatform.PlatformName() == TEXT("LinuxServer") || TargetPlatform.PlatformName() == TEXT("LinuxNoEditor"))
+		if (TargetPlatform.PlatformName() == TEXT("LinuxServer") || TargetPlatform.PlatformName() == TEXT("LinuxNoEditor") || TargetPlatform.PlatformName() == TEXT("Linux"))
 		{
 			Platform = TEXT("Linux");
 		}
diff --git a/Engine/Source/Developer/LauncherServices/Private/Launcher/LauncherCookCommands.h b/Engine/Source/Developer/LauncherServices/Private/Launcher/LauncherCookCommands.h
--- a/Engine/Source/Developer/LauncherServices/Private/Launcher/LauncherCookCommands.h
+++ b/Engine/Source/Developer/LauncherServices/Private/Launcher/LauncherCookCommands.h
@@ -111,7 +111,7 @@ public:
 		FString CommandLine;
 
 		FString Platform = TEXT("Win64");
-		if (TargetPlatform.PlatformName() == TEXT("LinuxServer") || TargetPlatform.PlatformName() == TEXT("LinuxNoEditor"))
+		if (TargetPlatform.PlatformName() == TEXT("LinuxServer") || TargetPlatform.PlatformName() == TEXT("LinuxNoEditor") || TargetPlatform.PlatformName() == TEXT("Linux"))
 		{
 			Platform = TEXT("Linux");
 		}
diff --git a/Engine/Source/Developer/LauncherServices/Private/Launcher/LauncherDeployCommands.h b/Engine/Source/Developer/LauncherServices/Private/Launcher/LauncherDeployCommands.h
--- a/Engine/Source/Developer/LauncherServices/Private/Launcher/LauncherDeployCommands.h
+++ b/Engine/Source/Developer/LauncherServices/Private/Launcher/LauncherDeployCommands.h
@@ -157,7 +157,7 @@ public:
 		}
 
 		FString Platform = TEXT("Win64");
-		if (TargetPlatform.PlatformName() == TEXT("LinuxServer") || TargetPlatform.PlatformName() == TEXT("LinuxNoEditor"))
+		if (TargetPlatform.PlatformName() == TEXT("LinuxServer") || TargetPlatform.PlatformName() == TEXT("LinuxNoEditor") || TargetPlatform.PlatformName() == TEXT("Linux"))
 		{
 			Platform = TEXT("Linux");
 		}
@@ -347,7 +347,7 @@ public:
 		}
 
 		FString Platform = TEXT("Win64");
-		if (TargetPlatform.PlatformName() == TEXT("LinuxServer") || TargetPlatform.PlatformName() == TEXT("LinuxNoEditor"))
+		if (TargetPlatform.PlatformName() == TEXT("LinuxServer") || TargetPlatform.PlatformName() == TEXT("LinuxNoEditor") || TargetPlatform.PlatformName() == TEXT("Linux"))
 		{
 			Platform = TEXT("Linux");
 		}
diff --git a/Engine/Source/Developer/LauncherServices/Private/Launcher/LauncherPackageCommands.h b/Engine/Source/Developer/LauncherServices/Private/Launcher/LauncherPackageCommands.h
--- a/Engine/Source/Developer/LauncherServices/Private/Launcher/LauncherPackageCommands.h
+++ b/Engine/Source/Developer/LauncherServices/Private/Launcher/LauncherPackageCommands.h
@@ -44,7 +44,7 @@ public:
 		// Determine if we should force PAK for the current configuration+platform+cook mode
 		if (ChainState.Profile->GetCookMode() == ELauncherProfileCookModes::ByTheBook)
 		{
-			if (ChainState.Profile->IsPackingWithUnrealPak() || (TargetPlatform.PlatformName() == TEXT("PS4")) || TargetPlatform.PlatformName().StartsWith(TEXT("Android")))
+			if (ChainState.Profile->IsPackingWithUnrealPak() || (TargetPlatform.PlatformName() == TEXT("LinuxServer")) || (TargetPlatform.PlatformName() == TEXT("PS4")) || TargetPlatform.PlatformName().StartsWith(TEXT("Android")))
 			{
 				CommandLine += TEXT(" -pak");
 			}
@@ -115,7 +115,7 @@ public:
 		}
 
 		FString Platform = TEXT("Win64");
-		if (TargetPlatform.PlatformName() == TEXT("LinuxServer") || TargetPlatform.PlatformName() == TEXT("LinuxNoEditor"))
+		if (TargetPlatform.PlatformName() == TEXT("LinuxServer") || TargetPlatform.PlatformName() == TEXT("LinuxNoEditor") || TargetPlatform.PlatformName() == TEXT("Linux"))
 		{
 			Platform = TEXT("Linux");
 		}
@@ -295,7 +295,7 @@ public:
 		}
 
 		FString Platform = TEXT("Win64");
-		if (TargetPlatform.PlatformName() == TEXT("LinuxServer") || TargetPlatform.PlatformName() == TEXT("LinuxNoEditor"))
+		if (TargetPlatform.PlatformName() == TEXT("LinuxServer") || TargetPlatform.PlatformName() == TEXT("LinuxNoEditor") || TargetPlatform.PlatformName() == TEXT("Linux"))
 		{
 			Platform = TEXT("Linux");
 		}
diff --git a/Engine/Source/Developer/LauncherServices/Private/Launcher/LauncherUATTask.h b/Engine/Source/Developer/LauncherServices/Private/Launcher/LauncherUATTask.h
--- a/Engine/Source/Developer/LauncherServices/Private/Launcher/LauncherUATTask.h
+++ b/Engine/Source/Developer/LauncherServices/Private/Launcher/LauncherUATTask.h
@@ -53,6 +53,8 @@ protected:
 		FString ExecutablePath = FPaths::ConvertRelativePathToFull(FPaths::EngineDir() + FString(TEXT("Build")) / TEXT("BatchFiles"));
 #if PLATFORM_MAC
 		FString Executable = TEXT("RunUAT.command");
+#elif PLATFORM_LINUX
+		FString Executable = TEXT("RunUAT.sh");
 #else
 		FString Executable = TEXT("RunUAT.bat");
 #endif
diff --git a/Engine/Source/Developer/Linux/LinuxTargetPlatform/Private/LinuxTargetPlatformPrivatePCH.h b/Engine/Source/Developer/Linux/LinuxTargetPlatform/Private/LinuxTargetPlatformPrivatePCH.h
--- a/Engine/Source/Developer/Linux/LinuxTargetPlatform/Private/LinuxTargetPlatformPrivatePCH.h
+++ b/Engine/Source/Developer/Linux/LinuxTargetPlatform/Private/LinuxTargetPlatformPrivatePCH.h
@@ -1,7 +1,7 @@
 // Copyright 1998-2014 Epic Games, Inc. All Rights Reserved.
 
 /*=============================================================================
-	LinuxTargetPlatformPrivatePCH.h: Pre-compiled header file for the AndroidTargetPlatform module.
+	LinuxTargetPlatformPrivatePCH.h: Pre-compiled header file for the LinuxTargetPlatform module.
 =============================================================================*/
 
 #pragma once
@@ -11,6 +11,7 @@
  *****************************************************************************/
 
 #include "Core.h"
+#include "ModuleInterface.h"
 #include "ModuleManager.h"
 #include "Settings.h"
 #include "Runtime/Core/Public/Linux/LinuxPlatformProperties.h"
@@ -23,7 +24,6 @@
 #include "TargetPlatform.h"
 #include "TargetPlatformBase.h"
 
-
 /* Private includes
  *****************************************************************************/
 
diff --git a/Engine/Source/Developer/SessionLauncher/Private/Widgets/Build/SSessionLauncherBuildPage.cpp b/Engine/Source/Developer/SessionLauncher/Private/Widgets/Build/SSessionLauncherBuildPage.cpp
--- a/Engine/Source/Developer/SessionLauncher/Private/Widgets/Build/SSessionLauncherBuildPage.cpp
+++ b/Engine/Source/Developer/SessionLauncher/Private/Widgets/Build/SSessionLauncherBuildPage.cpp
@@ -101,6 +101,8 @@ bool SSessionLauncherBuildPage::GenerateDSYMForProject( const FString& ProjectNa
     FString ExecutablePath = FPaths::ConvertRelativePathToFull(FPaths::EngineDir() + FString(TEXT("Build")) / TEXT("BatchFiles"));
 #if PLATFORM_MAC
     FString Executable = TEXT("RunUAT.command");
+#elif PLATFORM_LINUX
+    FString Executable = TEXT("RunUAT.sh");
 #else
     FString Executable = TEXT("RunUAT.bat");
 #endif
diff --git a/Engine/Source/Developer/ShaderFormatOpenGL/ShaderFormatOpenGL.Build.cs b/Engine/Source/Developer/ShaderFormatOpenGL/ShaderFormatOpenGL.Build.cs
--- a/Engine/Source/Developer/ShaderFormatOpenGL/ShaderFormatOpenGL.Build.cs
+++ b/Engine/Source/Developer/ShaderFormatOpenGL/ShaderFormatOpenGL.Build.cs
@@ -23,5 +23,10 @@ public class ShaderFormatOpenGL : ModuleRules
 			"OpenGL",
 			"HLSLCC"
 			);
+
+        if (Target.Platform == UnrealTargetPlatform.Linux)
+        {
+            AddThirdPartyPrivateStaticDependencies(Target, "SDL2");
+        }
 	}
 }
diff --git a/Engine/Source/Developer/TargetPlatform/TargetPlatform.Build.cs b/Engine/Source/Developer/TargetPlatform/TargetPlatform.Build.cs
--- a/Engine/Source/Developer/TargetPlatform/TargetPlatform.Build.cs
+++ b/Engine/Source/Developer/TargetPlatform/TargetPlatform.Build.cs
@@ -56,6 +56,22 @@ public class TargetPlatform : ModuleRules
 					DynamicallyLoadedModuleNames.Add("AudioFormatOgg");
 				}
 			}
+            else if (Target.Platform == UnrealTargetPlatform.Linux)
+            {
+                if (UEBuildConfiguration.bCompileLeanAndMeanUE == false)
+                {
+                    DynamicallyLoadedModuleNames.Add("TextureFormatDXT");
+                    DynamicallyLoadedModuleNames.Add("TextureFormatPVR");
+                }
+
+                DynamicallyLoadedModuleNames.Add("TextureFormatUncompressed");
+
+                if (UEBuildConfiguration.bCompileAgainstEngine)
+                {
+                    DynamicallyLoadedModuleNames.Add("AudioFormatOgg");
+                    DynamicallyLoadedModuleNames.Add("AudioFormatOpus");
+                }
+            }
 
 			if (UEBuildConfiguration.bCompileAgainstEngine && UEBuildConfiguration.bCompilePhysX)
 			{
diff --git a/Engine/Source/Developer/TextureFormatDXT/Private/TextureFormatDXT.cpp b/Engine/Source/Developer/TextureFormatDXT/Private/TextureFormatDXT.cpp
--- a/Engine/Source/Developer/TextureFormatDXT/Private/TextureFormatDXT.cpp
+++ b/Engine/Source/Developer/TextureFormatDXT/Private/TextureFormatDXT.cpp
@@ -65,6 +65,10 @@ struct FNVOutputHandler : public nvtt::OutputHandler
 		return true;
 	}
 
+    virtual void endImage()
+    {
+    }
+
 	uint8* Buffer;
 	uint8* BufferEnd;
 };
diff --git a/Engine/Source/Developer/TextureFormatPVR/Private/TextureFormatPVR.cpp b/Engine/Source/Developer/TextureFormatPVR/Private/TextureFormatPVR.cpp
--- a/Engine/Source/Developer/TextureFormatPVR/Private/TextureFormatPVR.cpp
+++ b/Engine/Source/Developer/TextureFormatPVR/Private/TextureFormatPVR.cpp
@@ -449,8 +449,12 @@ class FTextureFormatPVR : public ITextureFormat
 			bIsPVRTC2 ? 2 : 4);
 #if PLATFORM_MAC
 		FString CompressorPath(FPaths::EngineDir() + TEXT("Binaries/ThirdParty/ImgTec/PVRTexToolCLI"));
-#else
+#elif PLATFORM_LINUX
+		FString CompressorPath(FPaths::EngineDir() + TEXT("Binaries/ThirdParty/ImgTec/PVRTexToolCLI.lnx"));
+#elif PLATFORM_WINDOWS
 		FString CompressorPath(FPaths::EngineDir() + TEXT("Binaries/ThirdParty/ImgTec/PVRTexToolCLI.exe"));
+#else
+#error Unsupported platform
 #endif
 		UE_LOG(LogTemp, Log, TEXT("Running texturetool with '%s'"), *Params);
 
diff --git a/Engine/Source/Editor/EditorStyle/Private/SlateEditorStyle.cpp b/Engine/Source/Editor/EditorStyle/Private/SlateEditorStyle.cpp
--- a/Engine/Source/Editor/EditorStyle/Private/SlateEditorStyle.cpp
+++ b/Engine/Source/Editor/EditorStyle/Private/SlateEditorStyle.cpp
@@ -3196,7 +3196,7 @@ void FSlateEditorStyle::FStyle::SetupGraphEditorStyles()
 
 		Set( "GraphEditor.Default_16x", new IMAGE_BRUSH("Icons/icon_Blueprint_Node_16x", Icon16x16));
 		Set( "GraphEditor.EventGraph_16x", new IMAGE_BRUSH( "Icons/icon_Blueprint_EventGraph_16x", Icon16x16 ) );
-		Set( "GraphEditor.InterfaceFunction_16x", new IMAGE_BRUSH( "Icons/icon_Blueprint_InterfaceFunction_16x", Icon16x16 ) );
+		Set( "GraphEditor.InterfaceFunction_16x", new IMAGE_BRUSH( "Icons/icon_Blueprint_Interfacefunction_16x", Icon16x16 ) );
 		Set( "GraphEditor.Macro_16x", new IMAGE_BRUSH( "Icons/icon_Blueprint_Macro_16x", Icon16x16 ) );
 		Set( "GraphEditor.Function_16x", new IMAGE_BRUSH( "Icons/icon_Blueprint_NewFunction_16x", Icon16x16 ) );
 		Set( "GraphEditor.PotentialOverrideFunction_16x", new IMAGE_BRUSH( "Icons/icon_Blueprint_OverrideableFunction_16x", Icon16x16 ) );
@@ -3238,7 +3238,7 @@ void FSlateEditorStyle::FStyle::SetupGraphEditorStyles()
 		Set( "GraphEditor.Delegate_24x", new IMAGE_BRUSH( "Icons/icon_Blueprint_Delegate_24x", Icon24x24 ) );
 
 		Set( "GraphEditor.EventGraph_24x", new IMAGE_BRUSH( "Icons/icon_Blueprint_EventGraph_24x", Icon24x24 ) );
-		Set( "GraphEditor.InterfaceFunction_24x", new IMAGE_BRUSH( "Icons/icon_Blueprint_Interfacefunction_24x", Icon24x24 ) );
+		Set( "GraphEditor.InterfaceFunction_24x", new IMAGE_BRUSH( "Icons/icon_Blueprint_InterfaceFunction_24x", Icon24x24 ) );
 		Set( "GraphEditor.Macro_24x", new IMAGE_BRUSH( "Icons/icon_Blueprint_Macro_24x", Icon24x24 ) );
 		Set( "GraphEditor.Function_24x", new IMAGE_BRUSH( "Icons/icon_Blueprint_NewFunction_24x", Icon24x24 ) );
 		Set( "GraphEditor.PotentialOverrideFunction_24x", new IMAGE_BRUSH( "Icons/icon_Blueprint_OverrideableFunction_24x", Icon24x24 ) );
diff --git a/Engine/Source/Editor/Kismet/Private/FindInBlueprintManager.cpp b/Engine/Source/Editor/Kismet/Private/FindInBlueprintManager.cpp
--- a/Engine/Source/Editor/Kismet/Private/FindInBlueprintManager.cpp
+++ b/Engine/Source/Editor/Kismet/Private/FindInBlueprintManager.cpp
@@ -1021,13 +1021,7 @@ void FFindInBlueprintSearchManager::CleanCache()
 	for(auto& SearchValuePair : SearchMap)
 	{
 		// Here it builds the new map/array, clean of deleted content.
-		if( SearchArray[SearchValuePair.Value].bMarkedForDeletion || (SearchArray[SearchValuePair.Value].Blueprint && SearchArray[SearchValuePair.Value].Blueprint->HasAllFlags(RF_PendingKill)) )
-		{
-			// Item should not be kept, is is being destroyed and removed
-			int x = 5;
-			x = x;
-		}
-		else
+		if( !SearchArray[SearchValuePair.Value].bMarkedForDeletion && !(SearchArray[SearchValuePair.Value].Blueprint && SearchArray[SearchValuePair.Value].Blueprint->HasAllFlags(RF_PendingKill)) )
 		{
 			// Build the new map/array
 			NewSearchMap.Add(SearchValuePair.Key, NewSearchArray.Add(SearchArray[SearchValuePair.Value]) );
diff --git a/Engine/Source/Editor/LandscapeEditor/LandscapeEditor.Build.cs b/Engine/Source/Editor/LandscapeEditor/LandscapeEditor.Build.cs
--- a/Engine/Source/Editor/LandscapeEditor/LandscapeEditor.Build.cs
+++ b/Engine/Source/Editor/LandscapeEditor/LandscapeEditor.Build.cs
@@ -42,7 +42,7 @@ public class LandscapeEditor : ModuleRules
 
 		// KissFFT is used by the smooth tool.
 		if (UEBuildConfiguration.bCompileLeanAndMeanUE == false &&
-			(Target.Platform == UnrealTargetPlatform.Win64 || Target.Platform == UnrealTargetPlatform.Win32 || Target.Platform == UnrealTargetPlatform.Mac))
+			(Target.Platform == UnrealTargetPlatform.Win64 || Target.Platform == UnrealTargetPlatform.Win32 || Target.Platform == UnrealTargetPlatform.Mac || Target.Platform == UnrealTargetPlatform.Linux))
 		{
 			AddThirdPartyPrivateStaticDependencies(Target, "Kiss_FFT");
 		}
diff --git a/Engine/Source/Editor/MainFrame/Private/Frame/MainFrameActions.cpp b/Engine/Source/Editor/MainFrame/Private/Frame/MainFrameActions.cpp
--- a/Engine/Source/Editor/MainFrame/Private/Frame/MainFrameActions.cpp
+++ b/Engine/Source/Editor/MainFrame/Private/Frame/MainFrameActions.cpp
@@ -992,6 +992,9 @@ void FMainFrameActionCallbacks::CreateUatTask( const FString& CommandLine, const
 #if PLATFORM_WINDOWS
 	FString RunUATScriptName = TEXT("RunUAT.bat");
 	FString CmdExe = TEXT("cmd.exe");
+#elif PLATFORM_LINUX
+	FString RunUATScriptName = TEXT("RunUAT.sh");
+	FString CmdExe = TEXT("/bin/sh");
 #else
 	FString RunUATScriptName = TEXT("RunUAT.command");
 	FString CmdExe = TEXT("/bin/sh");
diff --git a/Engine/Source/Editor/TaskBrowser/Private/STaskBrowser.cpp b/Engine/Source/Editor/TaskBrowser/Private/STaskBrowser.cpp
--- a/Engine/Source/Editor/TaskBrowser/Private/STaskBrowser.cpp
+++ b/Engine/Source/Editor/TaskBrowser/Private/STaskBrowser.cpp
@@ -99,6 +99,8 @@ static bool EncryptBuffer(const uint8* SrcBuffer,const uint32 SrcLen,uint8* Dest
 		UE_LOG(LogTaskBrowser, Log, TEXT("CCCrypt failed w/ 0x%08x"), Status);
 	}
 	FMemory::Free(OutBuffer);
+#elif PLATFORM_LINUX
+	printf("STaskBrowser.cpp: LINUX EncryptBuffer()\n");
 #elif PLATFORM_WINDOWS
 	DATA_BLOB SourceBlob, EntropyBlob, FinalBlob;
 	// Set up the datablob to encrypt
@@ -189,6 +191,8 @@ static bool DecryptBuffer(const uint8* SrcBuffer,const uint32 SrcLen,uint8* Dest
 		UE_LOG(LogTaskBrowser, Log, TEXT("CCCrypt failed w/ 0x%08x"), Status);
 	}
 	FMemory::Free(OutBuffer);
+#elif PLATFORM_LINUX
+        printf("STaskBrowser.cpp: LINUX DecryptBuffer()\n");
 #elif PLATFORM_WINDOWS
 	DATA_BLOB SourceBlob, EntropyBlob, FinalBlob;
 	// Set up the datablob to encrypt
diff --git a/Engine/Source/Editor/UnrealEd/Private/Lightmass/Lightmass.cpp b/Engine/Source/Editor/UnrealEd/Private/Lightmass/Lightmass.cpp
--- a/Engine/Source/Editor/UnrealEd/Private/Lightmass/Lightmass.cpp
+++ b/Engine/Source/Editor/UnrealEd/Private/Lightmass/Lightmass.cpp
@@ -2464,15 +2464,26 @@ bool FLightmassProcessor::BeginRun()
 	};
 #endif
 #elif PLATFORM_LINUX
+#if UE_BUILD_DEBUG
 	const TCHAR* LightmassExecutable64 = TEXT("../Linux/UnrealLightmass-Linux-Debug");
 	const TCHAR* RequiredDependencyPaths64[] =
 	{
 		TEXT("../DotNET/Linux/AgentInterface.dll"),
-		TEXT("../Linux/UnrealLightmass-Core-Linux-Debug.so"),
-		TEXT("../Linux/UnrealLightmass-Projects-Linux-Debug.so"),
-		TEXT("../Linux/UnrealLightmass-SwarmInterface-Linux-Debug.so")
+		TEXT("../Linux/libUnrealLightmass-Core-Linux-Debug.so"),
+		TEXT("../Linux/libUnrealLightmass-Projects-Linux-Debug.so"),
+		TEXT("../Linux/libUnrealLightmass-SwarmInterface-Linux-Debug.so")
 	};
 #else
+	const TCHAR* LightmassExecutable64 = TEXT("../Linux/UnrealLightmass");
+	const TCHAR* RequiredDependencyPaths64[] =
+	{
+		TEXT("../DotNET/Linux/AgentInterface.dll"),
+		TEXT("../Linux/libUnrealLightmass-Core.so"),
+		TEXT("../Linux/libUnrealLightmass-Projects.so"),
+		TEXT("../Linux/libUnrealLightmass-SwarmInterface.so")
+	};
+#endif
+#else // platform
 #error "Unknown Lightmass platform"
 #endif
 	const int32 RequiredDependencyPaths64Count = ARRAY_COUNT(RequiredDependencyPaths64);
diff --git a/Engine/Source/Editor/UnrealEd/Private/MeshPaintRendering.cpp b/Engine/Source/Editor/UnrealEd/Private/MeshPaintRendering.cpp
--- a/Engine/Source/Editor/UnrealEd/Private/MeshPaintRendering.cpp
+++ b/Engine/Source/Editor/UnrealEd/Private/MeshPaintRendering.cpp
@@ -54,7 +54,7 @@ namespace MeshPaintRendering
 	};
 
 
-	IMPLEMENT_SHADER_TYPE( , TMeshPaintVertexShader, TEXT( "meshpaintvertexshader" ), TEXT( "Main" ), SF_Vertex);
+	IMPLEMENT_SHADER_TYPE( , TMeshPaintVertexShader, TEXT( "MeshPaintVertexShader" ), TEXT( "Main" ), SF_Vertex);
 
 
 
@@ -326,7 +326,7 @@ namespace MeshPaintRendering
 	};
 
 
-	IMPLEMENT_SHADER_TYPE( , TMeshPaintDilatePixelShader, TEXT( "MeshPaintDilatePixelShader" ), TEXT( "Main" ), SF_Pixel );
+	IMPLEMENT_SHADER_TYPE( , TMeshPaintDilatePixelShader, TEXT( "meshpaintdilatepixelshader" ), TEXT( "Main" ), SF_Pixel );
 
 
 	/** Mesh paint vertex format */
@@ -388,3 +388,4 @@ namespace MeshPaintRendering
 	}
 
 }
+
diff --git a/Engine/Source/Programs/UnrealHeaderTool/Private/Manifest.cpp b/Engine/Source/Programs/UnrealHeaderTool/Private/Manifest.cpp
--- a/Engine/Source/Programs/UnrealHeaderTool/Private/Manifest.cpp
+++ b/Engine/Source/Programs/UnrealHeaderTool/Private/Manifest.cpp
@@ -4,6 +4,7 @@
 #include "IScriptGeneratorPluginInterface.h"
 #include "Manifest.h"
 #include "Json.h"
+#include "Manifest.h"
 
 namespace
 {
diff --git a/Engine/Source/Programs/UnrealLightmass/Private/Launch/UnrealLightmass.cpp b/Engine/Source/Programs/UnrealLightmass/Private/Launch/UnrealLightmass.cpp
--- a/Engine/Source/Programs/UnrealLightmass/Private/Launch/UnrealLightmass.cpp
+++ b/Engine/Source/Programs/UnrealLightmass/Private/Launch/UnrealLightmass.cpp
@@ -489,7 +489,7 @@ int main(int argc, ANSICHAR* argv[])
 	}
 #endif
 
-#if PLATFORM_MAC
+#if PLATFORM_MAC || PLATFORM_LINUX
  	if ( true )
 #else
 	if ( FPlatformMisc::IsDebuggerPresent() )
diff --git a/Engine/Source/Programs/UnrealLightmass/Private/LightmassCore/Misc/LMDebug.cpp b/Engine/Source/Programs/UnrealLightmass/Private/LightmassCore/Misc/LMDebug.cpp
--- a/Engine/Source/Programs/UnrealLightmass/Private/LightmassCore/Misc/LMDebug.cpp
+++ b/Engine/Source/Programs/UnrealLightmass/Private/LightmassCore/Misc/LMDebug.cpp
@@ -75,6 +75,8 @@ void appHandleCriticalError()
 #endif
 #elif PLATFORM_MAC
 				FCString::Strcpy(PlatformName, TEXT("Mac"));
+#elif PLATFORM_LINUX
+				FCString::Strcpy(PlatformName, TEXT("Linux"));
 #endif
 				TCHAR LangExt[10];
 				FCString::Strcpy(LangExt, TEXT("English"));
diff --git a/Engine/Source/Programs/UnrealLightmass/Private/LightmassCore/Misc/LMThreading.cpp b/Engine/Source/Programs/UnrealLightmass/Private/LightmassCore/Misc/LMThreading.cpp
--- a/Engine/Source/Programs/UnrealLightmass/Private/LightmassCore/Misc/LMThreading.cpp
+++ b/Engine/Source/Programs/UnrealLightmass/Private/LightmassCore/Misc/LMThreading.cpp
@@ -11,7 +11,7 @@
 #include "AllowWindowsPlatformTypes.h"
 #endif // PLATFORM_WINDOWS
 
-#if PLATFORM_MAC
+#if PLATFORM_MAC || PLATFORM_LINUX
 #include <sys/time.h>	// for gettimeofday()
 #endif
 
diff --git a/Engine/Source/Runtime/Core/Private/Mac/MacPlatformProcess.cpp b/Engine/Source/Runtime/Core/Private/Mac/MacPlatformProcess.cpp
--- a/Engine/Source/Runtime/Core/Private/Mac/MacPlatformProcess.cpp
+++ b/Engine/Source/Runtime/Core/Private/Mac/MacPlatformProcess.cpp
@@ -643,8 +643,8 @@ const TCHAR* FMacPlatformProcess::UserDir()
 	if (!Result[0])
 	{
 		SCOPED_AUTORELEASE_POOL;
-		NSString *ApplicationSupportFolder = [NSSearchPathForDirectoriesInDomains(NSDocumentDirectory, NSUserDomainMask, YES) objectAtIndex: 0];
-		FPlatformString::CFStringToTCHAR((CFStringRef)ApplicationSupportFolder, Result);
+		NSString *DocumentsFolder = [NSSearchPathForDirectoriesInDomains(NSDocumentDirectory, NSUserDomainMask, YES) objectAtIndex: 0];
+		FPlatformString::CFStringToTCHAR((CFStringRef)DocumentsFolder, Result);
 		FCString::Strcat(Result, TEXT("/"));
 	}
 	return Result;
diff --git a/Engine/Source/Runtime/RHI/Private/Linux/LinuxDynamicRHI.cpp b/Engine/Source/Runtime/RHI/Private/Linux/LinuxDynamicRHI.cpp
--- a/Engine/Source/Runtime/RHI/Private/Linux/LinuxDynamicRHI.cpp
+++ b/Engine/Source/Runtime/RHI/Private/Linux/LinuxDynamicRHI.cpp
@@ -13,7 +13,7 @@ FDynamicRHI* PlatformCreateDynamicRHI()
 
 	if (!DynamicRHIModule->IsSupported())
 	{
-		FMessageDialog::Open(EAppMsgType::Ok, NSLOCTEXT("WindowsDynamicRHI", "RequiredOpenGL", "OpenGL 3.2 is required to run the engine."));
+		FMessageDialog::Open(EAppMsgType::Ok, NSLOCTEXT("LinuxDynamicRHI", "RequiredOpenGL", "OpenGL 3.2 is required to run the engine."));
 		FPlatformMisc::RequestExit(1);
 		DynamicRHIModule = NULL;
 	}
diff --git a/Engine/Source/Runtime/SlateRHIRenderer/Private/SlateRHIRenderer.cpp b/Engine/Source/Runtime/SlateRHIRenderer/Private/SlateRHIRenderer.cpp
--- a/Engine/Source/Runtime/SlateRHIRenderer/Private/SlateRHIRenderer.cpp
+++ b/Engine/Source/Runtime/SlateRHIRenderer/Private/SlateRHIRenderer.cpp
@@ -191,7 +191,7 @@ void FSlateRHIRenderer::Destroy()
 		BeginReleaseResource( It.Value() );
 	}
 	
-#if PLATFORM_WINDOWS || PLATFORM_MAC
+#if PLATFORM_WINDOWS || PLATFORM_MAC || PLATFORM_LINUX
 	if (CrashTrackerResource != nullptr)
 	{
 		BeginReleaseResource(CrashTrackerResource);
@@ -686,7 +686,7 @@ FIntRect FSlateRHIRenderer::SetupVirtualScreenBuffer(const bool bPrimaryWorkArea
 	FDisplayMetrics DisplayMetrics;
 	FSlateApplication::Get().GetDisplayMetrics(DisplayMetrics);
 
-#if !PLATFORM_WINDOWS && !PLATFORM_MAC
+#if !PLATFORM_WINDOWS && !PLATFORM_MAC && !PLATFORM_LINUX
 	ensureMsg(0, TEXT("This functionality is not valid for this platform"));	
 	return FIntRect(FIntPoint(0, 0), FIntPoint(DisplayMetrics.PrimaryDisplayWidth, DisplayMetrics.PrimaryDisplayHeight));
 #endif	
@@ -758,7 +758,7 @@ FIntRect FSlateRHIRenderer::SetupVirtualScreenBuffer(const bool bPrimaryWorkArea
 
 void FSlateRHIRenderer::CopyWindowsToVirtualScreenBuffer(const TArray<FString>& KeypressBuffer)
 {
-#if !PLATFORM_WINDOWS && !PLATFORM_MAC
+#if !PLATFORM_WINDOWS && !PLATFORM_MAC && !PLATFORM_LINUX
 	ensureMsg(0, TEXT("This functionality is not valid for this platform"));
 	return;
 #endif
diff --git a/Engine/Source/UE4Client.Target.cs b/Engine/Source/UE4Client.Target.cs
--- a/Engine/Source/UE4Client.Target.cs
+++ b/Engine/Source/UE4Client.Target.cs
@@ -57,7 +57,7 @@ public class UE4ClientTarget : TargetRules
         {
             return new List<UnrealTargetPlatform> { HostPlatform, UnrealTargetPlatform.IOS };
         }
-        return new List<UnrealTargetPlatform> { HostPlatform, UnrealTargetPlatform.Win32 /*, UnrealTargetPlatform.IOS, UnrealTargetPlatform.XboxOne, UnrealTargetPlatform.PS4 */};
+        return new List<UnrealTargetPlatform> { HostPlatform, UnrealTargetPlatform.Win32, UnrealTargetPlatform.Linux /*, UnrealTargetPlatform.IOS, UnrealTargetPlatform.XboxOne, UnrealTargetPlatform.PS4 */};
     }
 
     public override List<UnrealTargetConfiguration> GUBP_GetConfigs_MonolithicOnly(UnrealTargetPlatform HostPlatform, UnrealTargetPlatform Platform)
