Linux: updated scripts for (re)building TPS.

- We still have to do that (at least temporarily) as we don't bundle everything yet.

diff --git a/Engine/Build/BatchFiles/Linux/BuildThirdParty.sh b/Engine/Build/BatchFiles/Linux/BuildThirdParty.sh
--- a/Engine/Build/BatchFiles/Linux/BuildThirdParty.sh
+++ b/Engine/Build/BatchFiles/Linux/BuildThirdParty.sh
@@ -6,6 +6,7 @@
 SCRIPT_DIR=$(cd "$(dirname "$BASH_SOURCE")" ; pwd)
 TOP_DIR=$(cd $SCRIPT_DIR/../../.. ; pwd)
 cd ${TOP_DIR}
+mkdir -p Binaries/Linux/
 set -e
 
 MAKE_ARGS=-j4
@@ -36,23 +37,6 @@ P4Open()
   done
 }
 
-BuildJemalloc()
-{
-  echo "building jemalloc"
-  set -x
-  cd Source/ThirdParty/jemalloc/build
-  tar xf jemalloc-3.4.0.tar.bz2
-  cd jemalloc-3.4.0
-  ./configure --with-mangling --with-jemalloc-prefix=je_340_
-  make
-  local LIB_DIR=../../lib/Linux/$TARGET_ARCH/separate
-  P4Open $LIB_DIR/libjemalloc.a
-  mkdir -p $LIB_DIR
-  cp lib/* $LIB_DIR/
-  cp lib/libjemalloc.so.1 ${TOP_DIR}/Binaries/Linux
-  set +x
-}
-
 BuildOpus()
 {
   echo "building libOpus"
@@ -69,28 +53,6 @@ BuildOpus()
   set +x
 }
 
-BuildPNG()
-{
-  echo "building libPNG"
-  set -x
-  cd Source/ThirdParty/libPNG/libPNG-1.5.2
-  cp --remove-destination scripts/makefile.linux Makefile
-  cp --remove-destination scripts/pnglibconf.h.prebuilt pnglibconf.h
-  P4Open pnglibconf.h
-  make $MAKE_ARGS
-
-  local LIB_DIR=lib/Linux/$TARGET_ARCH
-  mkdir -p $LIB_DIR
-  P4Open $LIB_DIR/libpng.a
-  cp --remove-destination libpng.a $LIB_DIR/
-  # (amigo) this dylib will mess with our build
-  # instead we want to use system's libpng to avoid asserts
-  #cp libpng15.so $LIB_DIR/libpng.so
-  mkdir -p ${TOP_DIR}/Binaries/Linux/
-  #cp -a libpng15.so* ${TOP_DIR}/Binaries/Linux/
-  set +x
-}
-
 BuildOgg()
 {
   echo "building Ogg"
@@ -121,7 +83,7 @@ BuildVorbis()
   make $MAKE_ARGS
 
   local LIB_DIR=lib/Linux/$TARGET_ARCH
-  P4Open $LIB_DIR/libvorbis.a $LIB_DIR/libvorbisfile.a
+  #P4Open $LIB_DIR/libvorbis.a $LIB_DIR/libvorbisfile.a
   mkdir -p $LIB_DIR
   cp lib/.libs/libvorbis*.a $LIB_DIR/
   set +x
@@ -134,6 +96,7 @@ BuildHLSLCC()
   set -x
   cd Source/ThirdParty/hlslcc
   P4Open hlslcc/bin/Linux/hlslcc_64
+  P4Open hlslcc/lib/Linux/libhlslcc.a
   make $MAKE_ARGS
   set +x
 }
@@ -158,21 +121,6 @@ BuildMcpp()
   set +x
 }
 
-BuildFreeType()
-{
-  echo "building freetype"
-  set -x
-  cd Source/ThirdParty/FreeType2/FreeType2-2.4.12/src
-  pwd
-  make $MAKE_ARGS -f ../Builds/Linux/makefile $*
-
-  local LIB_DIR=../Lib/Linux/$TARGET_ARCH
-  P4Open $LIB_DIR/libfreetype2412.a
-  mkdir -p $LIB_DIR
-  cp --remove-destination libfreetype* $LIB_DIR/
-  set +x
-}
-
 BuildICU()
 {
   echo "building libICU"
@@ -267,16 +215,17 @@ BuildnvTextureTools()
   cd Source/ThirdParty/nvTextureTools/nvTextureTools-2.0.8
   mkdir -p lib/Linux/$TARGET_ARCH
   cd src
-  ./configure
+  P4Open configure
+  chmod +x ./configure
+  CFLAGS=-fPIC ./configure
   make
 
   local LIB_DIR=../lib/Linux/$TARGET_ARCH
-  cp --remove-destination build-debug/src/nvcore/libnvcore.so $LIB_DIR/
-  cp --remove-destination build-debug/src/nvimage/libnvimage.so $LIB_DIR/
-  cp --remove-destination build-debug/src/nvmath/libnvmath.so $LIB_DIR/
-  cp --remove-destination build-debug/src/nvtt/libnvtt.so $LIB_DIR/
-  cp --remove-destination build-debug/src/nvthread/libnvthread.a $LIB_DIR/
-  cp --remove-destination build-debug/src/nvtt/squish/libsquish.a $LIB_DIR/
+  cp --remove-destination build/src/nvcore/libnvcore.so $LIB_DIR/
+  cp --remove-destination build/src/nvimage/libnvimage.so $LIB_DIR/
+  cp --remove-destination build/src/nvmath/libnvmath.so $LIB_DIR/
+  cp --remove-destination build/src/nvtt/libnvtt.so $LIB_DIR/
+  cp --remove-destination build/src/nvtt/squish/libsquish.a $LIB_DIR/
   cp --remove-destination $LIB_DIR/*.so ${TOP_DIR}/Binaries/Linux/
   set +x
 }
@@ -285,20 +234,23 @@ BuildSDL2()
 {
   echo "building SDL2"
   set -x
-  cd Source/ThirdParty/SDL2/SDL2-2.0.3/build/SDL2-2.0.3
-  P4Open configure
+  cd Source/ThirdParty/SDL2/SDL-gui-backend/build
+  # always do it here for now, since there's no UpdateDeps step for builds from Perforce (binaries are part of the repo).
+  tar xjf *tar.bz2
+  cd SDL-gui-backend
   chmod +x configure
   chmod +x autogen.sh
   chmod +x Android.mk
-  ./configure
+  CFLAGS=-fPIC ./configure
   make $MAKE_ARGS
-  cp --remove-destination include/* ../../include/
   local LIB_DIR=../../lib/Linux/$TARGET_ARCH
-  P4Open $LIB_DIR/libEpicSDL2.a
   mkdir -p $LIB_DIR
-  cp --remove-destination build/.libs/libSDL2.a $LIB_DIR/libEpicSDL2.a
-  cp build/.libs/libSDL2.so $LIB_DIR/libEpicSDL2.so
-  cp build/.libs/libSDL2.so ${TOP_DIR}/Binaries/Linux/libSDL2-2.0.so.0
+  P4Open $LIB_DIR/libSDL2.a
+  P4Open $LIB_DIR/libSDL2.so
+  P4Open ${TOP_DIR}/Binaries/Linux/libSDL2-2.0.so.0
+  cp --remove-destination build/.libs/libSDL2.a $LIB_DIR/libSDL2.a
+  cp --remove-destination build/.libs/libSDL2.so $LIB_DIR/libSDL2.so
+  cp --remove-destination build/.libs/libSDL2.so ${TOP_DIR}/Binaries/Linux/libSDL2-2.0.so.0
   set +x
 }
 
@@ -358,15 +310,13 @@ fi
 
 rm -f ${SCRIPT_DIR}/BuildThirdParty.log
 if [ -z "$1" ]; then
-  Run BuildJemalloc
-  Run BuildOpus
-  Run BuildOgg
-  Run BuildVorbis
-  Run BuildPNG
+  # use bundled Run BuildOpus
+  # use bundled Run BuildOgg
+  # use bundled Run BuildVorbis
   Run BuildHLSLCC
   Run BuildMcpp
-  Run BuildFreeType
-  Run BuildLND
+  # use bundled Run BuildFreeType
+  #on hold for now Run BuildLND
   Run BuildForsythTriOO
   Run BuildnvTriStrip
   Run BuildnvTextureTools
diff --git a/Engine/Build/BatchFiles/Linux/GenerateProjectFiles.sh b/Engine/Build/BatchFiles/Linux/GenerateProjectFiles.sh
--- a/Engine/Build/BatchFiles/Linux/GenerateProjectFiles.sh
+++ b/Engine/Build/BatchFiles/Linux/GenerateProjectFiles.sh
@@ -26,6 +26,7 @@ if [ "$(lsb_release --id)" = "Distributor ID:	Ubuntu" -o "$(lsb_release --id)" =
     libmono-system-web-extensions4.0-cil
     libmono-system-management4.0-cil
     libmono-system-xml-linq4.0-cil
+    libmono-corlib4.0-cil
     libogg-dev"
 
   for DEP in $DEPS; do
@@ -38,24 +39,37 @@ if [ "$(lsb_release --id)" = "Distributor ID:	Ubuntu" -o "$(lsb_release --id)" =
   done
 fi
 
-# Fixes for case sensitive filesystem.
-for BASE in Content/Editor/Slate Content/Slate Documentation/Source/Shared/Icons; do
-  find $BASE -name "*.PNG" | while read PNG_UPPER; do
-    png_lower="$(echo "$PNG_UPPER" | sed 's/.PNG$/.png/')"
-    if [ ! -f $png_lower ]; then
-      PNG_UPPER=$(basename $PNG_UPPER)
-      echo "$png_lower -> $PNG_UPPER"
-      # link, and not move, to make it usable with Perforce workspaces
-      ln -sf `basename "$PNG_UPPER"` "$png_lower"
-    fi
-  done
-done
-
 set -x
 xbuild Source/Programs/UnrealBuildTool/UnrealBuildTool_Mono.csproj \
   /verbosity:quiet /nologo \
   /p:TargetFrameworkVersion=v4.0 \
   /p:Configuration="Development"
 
+xbuild Source/Programs/AutomationTool/AutomationTool_Mono.csproj \
+  /verbosity:quiet /nologo \
+  /p:TargetFrameworkVersion=v4.0 \
+  /p:Configuration="Development"
+
+xbuild Source/Programs/AutomationTool/Scripts/AutomationScripts.Automation.csproj \
+  /verbosity:quiet /nologo \
+  /p:TargetFrameworkVersion=v4.0 \
+  /p:Configuration="Development"
+
+xbuild Source/Programs/AutomationTool/Linux/Linux.Automation.csproj \
+  /verbosity:quiet /nologo \
+  /p:TargetFrameworkVersion=v4.0 \
+  /p:Configuration="Development"
+
+xbuild Source/Programs/AutomationTool/Android/Android.Automation.csproj \
+  /verbosity:quiet /nologo \
+  /p:TargetFrameworkVersion=v4.0 \
+  /p:Configuration="Development"
+
+xbuild Source/Programs/AutomationTool/HTML5/HTML5.Automation.csproj \
+  /verbosity:quiet /nologo \
+  /p:TargetFrameworkVersion=v4.0 \
+  /p:Configuration="Development"
+
 # pass all parameters to UBT
 mono Binaries/DotNET/UnrealBuildTool.exe -makefile "$@"
+set +x
