Linux: add fix-ups for case-sensitive fs again.

- Needed until p4 repo issues are resolved.

diff --git a/Engine/Build/BatchFiles/Linux/GenerateProjectFiles.sh b/Engine/Build/BatchFiles/Linux/GenerateProjectFiles.sh
--- a/Engine/Build/BatchFiles/Linux/GenerateProjectFiles.sh
+++ b/Engine/Build/BatchFiles/Linux/GenerateProjectFiles.sh
@@ -1,4 +1,4 @@
-#!/bin/sh
+#!/bin/bash
 
 SCRIPT_DIR=$(cd "$(dirname "$BASH_SOURCE")" ; pwd)
 
@@ -39,6 +39,38 @@ if [ "$(lsb_release --id)" = "Distributor ID:	Ubuntu" -o "$(lsb_release --id)" =
   done
 fi
 
+# args: wrong filename, correct filename
+# expects to be in Engine folder
+CreateLinkIfNoneExists()
+{
+    WrongName=$1
+    CorrectName=$2
+
+    pushd `dirname $CorrectName`
+    if [ ! -f `basename $CorrectName` ]; then
+      echo "$WrongName -> $CorrectName"
+      ln -sf $WrongName `basename $CorrectName`
+    fi
+    popd
+}
+
+
+
+# Fixes for case sensitive filesystem.
+for BASE in Content/Editor/Slate Content/Slate Documentation/Source/Shared/Icons; do
+  find $BASE -name "*.PNG" | while read PNG_UPPER; do
+    png_lower="$(echo "$PNG_UPPER" | sed 's/.PNG$/.png/')"
+    if [ ! -f $png_lower ]; then
+      PNG_UPPER=$(basename $PNG_UPPER)
+      echo "$png_lower -> $PNG_UPPER"
+      # link, and not move, to make it usable with Perforce workspaces
+      ln -sf `basename "$PNG_UPPER"` "$png_lower"
+    fi
+  done
+done
+
+CreateLinkIfNoneExists ../../engine/shaders/Fxaa3_11.usf  ../Engine/Shaders/Fxaa3_11.usf
+
 set -x
 xbuild Source/Programs/UnrealBuildTool/UnrealBuildTool_Mono.csproj \
   /verbosity:quiet /nologo \
diff --git a/GenerateProjectFiles.sh b/GenerateProjectFiles.sh
--- a/GenerateProjectFiles.sh
+++ b/GenerateProjectFiles.sh
@@ -16,6 +16,7 @@ if [ "$(uname)" = "Darwin" ]; then
 	sh ./GenerateLLDBInit.sh
 	sh ./GenerateProjectFiles.sh $@
 else
+    # assume (GNU/)Linux
 	cd Engine/Build/BatchFiles/Linux
-	sh ./GenerateProjectFiles.sh $@
+	bash ./GenerateProjectFiles.sh $@
 fi
