Manual smoke script for QA.

diff --git a/Engine/Build/BatchFiles/Linux/QASmokeManual.sh b/Engine/Build/BatchFiles/Linux/QASmokeManual.sh
new file mode 100755
--- /dev/null
+++ b/Engine/Build/BatchFiles/Linux/QASmokeManual.sh
@@ -0,0 +1,104 @@
+#!/bin/bash
+
+STEP=0
+CURRENT_DIR=$(cd "$(dirname "$BASH_SOURCE")"; pwd)
+LOG_FILE=$CURRENT_DIR/QASmoke.log
+SUCCESS_MESSAGE="***** SMOKE SUCCESSFUL *****"
+
+Header()
+{
+  echo "Ran on `date` on `hostname`"
+  echo "Original location: $LOG_FILE"
+}
+
+Sync()
+{
+  echo "Syncing the repo"
+  set -x
+  p4 sync ...#head 2>&1
+  set +x
+}
+
+GenerateProjectFiles()
+{
+  echo "Generating project files"
+  set -x
+  bash GenerateProjectFiles.sh
+  set +x
+}
+
+BuildTarget()
+{
+  echo "Cleaning $1"
+  set -x
+  make $1 ARGS="-clean"
+  set +x
+
+  echo "Attempting timed build of $1"
+  set -x
+  time make $1
+  set +x
+  echo "$1 was successful"
+}
+
+CompileTPS()
+{
+  # FIXME: get rid of this step
+  echo "Recompiling editor dependencies locally"
+  set -x
+  pushd Engine/Build/BatchFiles/Linux
+  export VERBOSE=1
+  ./BuildThirdParty.sh
+  popd
+  set +x
+}
+
+Run()
+{
+  STEP=$((STEP+1))
+  echo "Step $STEP: $1 ${@:2}"
+  echo "------------------------------------------------------------------------" >> $LOG_FILE
+  echo "  Step $STEP $1 ${@:2}" >> $LOG_FILE
+  echo "------------------------------------------------------------------------" >> $LOG_FILE
+  $1 ${@:2} >> $LOG_FILE 2>&1
+}
+
+if [ -z "$1" ]; then
+  echo "---------------"
+  echo "This is a script for QA to manually smoke a branch"
+  echo ""
+  echo "Successful run should end with $SUCCESS_MESSAGE"
+  echo "Any other messages indicate a problem - see `basename $LOG_FILE`"
+  echo ""
+  echo "---------------"
+  
+
+  pushd ../../../.. >> /dev/null
+  rm -f $LOG_FILE
+  Run Header
+  Run Sync
+  Run GenerateProjectFiles
+
+  # build usual targets, roughly ascending by build time
+  Run BuildTarget BlankProgram
+  Run BuildTarget UnrealPak
+  Run BuildTarget SlateViewer 
+  Run BuildTarget UE4Client
+  Run BuildTarget UE4Game
+  Run BuildTarget UE4Server 
+
+  # rebuild dependencies locally for editor targets (extra safety still)
+  Run CompileTPS
+
+  # build editor targets, roughly ascending by build time
+  Run BuildTarget ShaderCompileWorker
+  Run BuildTarget UnrealFrontend
+  Run BuildTarget UE4Editor
+  popd >> /dev/null
+
+else
+  time Run $1
+fi
+
+echo "***** Smoke successful *****"
+
