Removed Sync step.
Log files are labeled with timestamp and are not deleted on each run.
Added set -e to allow early exit on errors.

diff --git a/Engine/Build/BatchFiles/Linux/QASmokeManual.sh b/Engine/Build/BatchFiles/Linux/QASmokeManual.sh
--- a/Engine/Build/BatchFiles/Linux/QASmokeManual.sh
+++ b/Engine/Build/BatchFiles/Linux/QASmokeManual.sh
@@ -2,7 +2,8 @@
 
 STEP=0
 CURRENT_DIR=$(cd "$(dirname "$BASH_SOURCE")"; pwd)
-LOG_FILE=$CURRENT_DIR/QASmoke.log
+TIMESTAMP=`date +%m-%d-%d_%H.%M.%s`
+LOG_FILE=$CURRENT_DIR/"QASmoke_$TIMESTAMP.log"
 SUCCESS_MESSAGE="***** SMOKE SUCCESSFUL *****"
 
 Header()
@@ -11,14 +12,6 @@ Header()
   echo "Original location: $LOG_FILE"
 }
 
-Sync()
-{
-  echo "Syncing the repo"
-  set -x
-  p4 sync ...#head 2>&1
-  set +x
-}
-
 GenerateProjectFiles()
 {
   echo "Generating project files"
@@ -63,6 +56,8 @@ Run()
   $1 ${@:2} >> $LOG_FILE 2>&1
 }
 
+# Main
+set -e
 if [ -z "$1" ]; then
   echo "---------------"
   echo "This is a script for QA to manually smoke a branch"
@@ -74,9 +69,7 @@ if [ -z "$1" ]; then
   
 
   pushd ../../../.. >> /dev/null
-  rm -f $LOG_FILE
   Run Header
-  Run Sync
   Run GenerateProjectFiles
 
   # build usual targets, roughly ascending by build time
@@ -100,5 +93,5 @@ else
   time Run $1
 fi
 
-echo "***** Smoke successful *****"
+echo "$SUCCESS_MESSAGE"
 
