Fixed timestamp formatting for Linux smoke test log.

diff --git a/Engine/Build/BatchFiles/Linux/QASmokeManual.sh b/Engine/Build/BatchFiles/Linux/QASmokeManual.sh
--- a/Engine/Build/BatchFiles/Linux/QASmokeManual.sh
+++ b/Engine/Build/BatchFiles/Linux/QASmokeManual.sh
@@ -2,7 +2,7 @@
 
 STEP=0
 CURRENT_DIR=$(cd "$(dirname "$BASH_SOURCE")"; pwd)
-TIMESTAMP=`date +%m-%d-%d_%H.%M.%s`
+TIMESTAMP=`date +%y-%m-%d_%H.%M.%S`
 LOG_FILE=$CURRENT_DIR/"QASmoke_$TIMESTAMP.log"
 SUCCESS_MESSAGE="***** SMOKE SUCCESSFUL *****"
 
