Fix BuildThirdParty.sh and QASmokeManual.sh

- Check out files it is modifying.
- Build LND before building SlateViewer.

diff --git a/Engine/Build/BatchFiles/Linux/BuildThirdParty.sh b/Engine/Build/BatchFiles/Linux/BuildThirdParty.sh
--- a/Engine/Build/BatchFiles/Linux/BuildThirdParty.sh
+++ b/Engine/Build/BatchFiles/Linux/BuildThirdParty.sh
@@ -188,6 +188,7 @@ BuildForsythTriOO()
   make
   cmake -DCMAKE_BUILD_TYPE=Release ../../
   make
+  P4Open ../../Lib/Linux/$TARGET_ARCH/libForsythTriOptimizer*
   cp --remove-destination libForsythTriOptimizer* ../../Lib/Linux/$TARGET_ARCH/
   set +x
 }
@@ -204,6 +205,7 @@ BuildnvTriStrip()
   make
   cmake -DCMAKE_BUILD_TYPE=Release ../../
   make
+  P4Open ../../Lib/Linux/$TARGET_ARCH/libnvtristrip*
   cp --remove-destination libnvtristrip* ../../Lib/Linux/$TARGET_ARCH/
   set +x
 }
@@ -221,10 +223,15 @@ BuildnvTextureTools()
   make
 
   local LIB_DIR=../lib/Linux/$TARGET_ARCH
+  P4Open $LIB_DIR/libnvcore.so
   cp --remove-destination build/src/nvcore/libnvcore.so $LIB_DIR/
+  P4Open $LIB_DIR/libnvimage.so
   cp --remove-destination build/src/nvimage/libnvimage.so $LIB_DIR/
+  P4Open $LIB_DIR/libnvmath.so
   cp --remove-destination build/src/nvmath/libnvmath.so $LIB_DIR/
+  P4Open $LIB_DIR/libnvtt.so
   cp --remove-destination build/src/nvtt/libnvtt.so $LIB_DIR/
+  P4Open $LIB_DIR/libsquish.so
   cp --remove-destination build/src/nvtt/squish/libsquish.a $LIB_DIR/
   cp --remove-destination $LIB_DIR/*.so ${TOP_DIR}/Binaries/Linux/
   set +x
@@ -265,6 +272,8 @@ Buildcoremod()
   make $MAKE_ARGS
   local LIB_DIR=lib/Linux/$TARGET_ARCH
   mkdir -p $LIB_DIR
+
+  P4Open $LIB_DIR/libcoremodLinux.a
   cp lib/libxmp-coremod.a $LIB_DIR/libcoremodLinux.a
   set +x
 }
