Minor fix for BuildThirdParty script.

diff --git a/Engine/Build/BatchFiles/Linux/BuildThirdParty.sh b/Engine/Build/BatchFiles/Linux/BuildThirdParty.sh
--- a/Engine/Build/BatchFiles/Linux/BuildThirdParty.sh
+++ b/Engine/Build/BatchFiles/Linux/BuildThirdParty.sh
@@ -100,7 +100,7 @@ BuildHLSLCC()
   cd Source/ThirdParty/hlslcc
   P4Open hlslcc/bin/Linux/hlslcc_64
   P4Open hlslcc/lib/Linux/libhlslcc.a
-  cd projects/Linux
+  cd hlslcc/projects/Linux
   make $MAKE_ARGS
   set +x
 }
