Correct name-casing fix for a different misspelling in 4.5 branch.

diff --git a/Engine/Build/BatchFiles/Linux/GenerateProjectFiles.sh b/Engine/Build/BatchFiles/Linux/GenerateProjectFiles.sh
--- a/Engine/Build/BatchFiles/Linux/GenerateProjectFiles.sh
+++ b/Engine/Build/BatchFiles/Linux/GenerateProjectFiles.sh
@@ -47,7 +47,7 @@ CreateLinkIfNoneExists()
     CorrectName=$2
 
     pushd `dirname $CorrectName`
-    if [ ! -f `basename $CorrectName` ]; then
+    if [ ! -f `basename $CorrectName` ] && [ -f $WrongName ]; then
       echo "$WrongName -> $CorrectName"
       ln -sf $WrongName `basename $CorrectName`
     fi
@@ -70,6 +70,7 @@ for BASE in Content/Editor/Slate Content/Slate Documentation/Source/Shared/Icons
 done
 
 CreateLinkIfNoneExists ../../engine/shaders/Fxaa3_11.usf  ../Engine/Shaders/Fxaa3_11.usf
+CreateLinkIfNoneExists ../../Engine/shaders/Fxaa3_11.usf  ../Engine/Shaders/Fxaa3_11.usf
 
 set -x
 xbuild Source/Programs/UnrealBuildTool/UnrealBuildTool_Mono.csproj \
