Reworked workflow for native Linux builds.

- The flow should be now as simple as ./GenerateProjectFiles.sh + make <target>
   (you will need OAUTH_TOKEN variable set and you may need to answer "y" occasionally - when removing existing files).

- Archvies should be now automatically (re-)downloaded by GetAssets.py script by nycat. They won't be
  re-downloaded if the files on github have not been updated since the last download.

- BuildThirdParty.sh is invoked by the above script and only builds LND (PR #469).
   It can still be invoked manually to build whatever if desired so, but watch out -
   you will need to take care not to build against too new glibc so you can distribute your game's binaries,
   and you will need to match options we built libraries with (e.g. Xinerama/XRandR headers for SDL2).

- UpdateDeps.sh now should only be called by GenerateProjectFiles.sh. It does a lot less
   work and shares variables with GetAssets.py

Resolves TTP #347988.

diff --git a/Engine/Build/BatchFiles/Linux/BuildThirdParty.sh b/Engine/Build/BatchFiles/Linux/BuildThirdParty.sh
--- a/Engine/Build/BatchFiles/Linux/BuildThirdParty.sh
+++ b/Engine/Build/BatchFiles/Linux/BuildThirdParty.sh
@@ -40,6 +40,45 @@ P4Open()
   done
 }
 
+BuildZ()
+{
+  echo "building zlib"
+  set -x
+  cd Source/ThirdParty/zlib/zlib-1.2.5/build
+  tar xf zlib-1.2.5.tar.gz
+  cd zlib-1.2.5
+  CFLAGS=-fPIC ./configure
+  make $MAKE_ARGS
+  local LIB_DIR=../../Lib/Linux/$TARGET_ARCH
+  P4Open $LIB_DIR/libz.a
+  P4Open $LIB_DIR/libz_fPIC.a
+  mkdir -p $LIB_DIR
+  cp libz.a $LIB_DIR/libz_fPIC.a
+  set +x
+}
+
+BuildJemalloc()
+{
+  echo "building jemalloc"
+  set -x
+  cd Source/ThirdParty/jemalloc/build
+  tar xf jemalloc-3.6.0.tar.bz2
+  cd jemalloc-3.6.0
+  ./configure --with-mangling --with-jemalloc-prefix=je_
+  make $MAKE_ARGS
+  local LIB_DIR=../../lib/Linux/$TARGET_ARCH
+  local INC_DIR=../../include/Linux/$TARGET_ARCH
+  P4Open $LIB_DIR/libjemalloc.a
+  P4Open $LIB_DIR/libjemalloc_pic.a
+  mkdir -p $LIB_DIR
+  cp -rp lib/* $LIB_DIR
+  mkdir -p $INC_DIR
+  P4Open $INC_DIR/jemalloc_defs.h
+  P4Open $INC_DIR/jemalloc.h
+  cp -rp include/jemalloc/* $INC_DIR
+  set +x
+}
+
 BuildOpus()
 {
   echo "building libOpus"
@@ -51,8 +90,8 @@ BuildOpus()
   make $MAKE_ARGS
   local LIB_DIR=Linux/$TARGET_ARCH
   mkdir -p $LIB_DIR
-  cp .libs/libopus.a .libs/libopus*.so* $LIB_DIR
-  cp $LIB_DIR/libopus.so.0 ${TOP_DIR}/Binaries/Linux/
+  P4Open $LIB_DIR/libopus_fPIC.a
+  cp .libs/libopus.a $LIB_DIR/libopus_fPIC.a
   set +x
 }
 
@@ -68,8 +107,10 @@ BuildOgg()
 
   local LIB_DIR=lib/Linux/$TARGET_ARCH
   P4Open $LIB_DIR/libogg.a
+  P4Open $LIB_DIR/libogg_fPIC.a
   mkdir -p $LIB_DIR
-  cp src/.libs/* $LIB_DIR
+  cp src/.libs/libogg.a $LIB_DIR
+  cp $LIB_DIR/libogg.a $LIB_DIR/libogg_fPIC.a
   set +x
 }
 
@@ -86,9 +127,13 @@ BuildVorbis()
   make $MAKE_ARGS
 
   local LIB_DIR=lib/Linux/$TARGET_ARCH
-  #P4Open $LIB_DIR/libvorbis.a $LIB_DIR/libvorbisfile.a
+  P4Open $LIB_DIR/libvorbis.a $LIB_DIR/libvorbisfile.a $LIB_DIR/libvorbis_fPIC.a $LIB_DIR/libvorbisfile_fPIC.a
+  P4Open $LIB_DIR/libvorbisenc.a $LIB_DIR/libvorbisenc_fPIC.a 
   mkdir -p $LIB_DIR
   cp lib/.libs/libvorbis*.a $LIB_DIR/
+  cp $LIB_DIR/libvorbis.a libvorbis_fPIC.a
+  cp $LIB_DIR/libvorbisfile.a libvorbisfile_fPIC.a
+  cp $LIB_DIR/libvorbisenc.a libvorbisenc_fPIC.a
   set +x
   cd -
 }
@@ -98,7 +143,7 @@ BuildHLSLCC()
   echo "building hlslcc"
   set -x
   cd Source/ThirdParty/hlslcc
-  P4Open hlslcc/bin/Linux/hlslcc_64
+  # not building anymore P4Open hlslcc/bin/Linux/hlslcc_64
   P4Open hlslcc/lib/Linux/libhlslcc.a
   cd hlslcc/projects/Linux
   make $MAKE_ARGS
@@ -118,6 +163,7 @@ BuildMcpp()
   local LIB_DIR=lib/Linux/$TARGET_ARCH
   P4Open $LIB_DIR/libmcpp.a
   P4Open $LIB_DIR/libmcpp.so
+  P4Open ${TOP_DIR}/Binaries/Linux/libmcpp.so.0
   mkdir -p $LIB_DIR
   cp --remove-destination ./src/.libs/libmcpp.a $LIB_DIR/
   cp --remove-destination ./src/.libs/libmcpp.so $LIB_DIR/
@@ -125,6 +171,22 @@ BuildMcpp()
   set +x
 }
 
+BuildFreeType()
+{
+  echo "building freetype"
+  set -x
+  cd Source/ThirdParty/FreeType2/FreeType2-2.4.12/src
+  pwd
+  make $MAKE_ARGS -f ../Builds/Linux/makefile $*
+
+  local LIB_DIR=../Lib/Linux/$TARGET_ARCH
+  P4Open $LIB_DIR/libfreetype2412.a
+  P4Open $LIB_DIR/libfreetype2412_fPIC.a
+  cp --remove-destination libfreetype2412.a $LIB_DIR/libfreetype2412.a
+  cp $LIB_DIR/libfreetype2412.a $LIB_DIR/libfreetype2412_fPIC.a
+  set +x
+}
+
 BuildICU()
 {
   echo "building libICU"
@@ -153,7 +215,10 @@ BuildLND()
 {
   echo "building LinuxNativeDialogs"
   set -x
-  cd Source/ThirdParty/LinuxNativeDialogs/UELinuxNativeDialogs/build
+  cd Source/ThirdParty/LinuxNativeDialogs/UELinuxNativeDialogs
+  rm -rf build
+  mkdir build
+  cd build
   cmake ..
   make $MAKE_ARGS
   local LIB_DIR=../lib/Linux/$TARGET_ARCH/
@@ -185,6 +250,7 @@ BuildForsythTriOO()
   echo "building ForsythTriOO"
   set -x
   cd Source/ThirdParty/ForsythTriOO
+  rm -rf Build/Linux
   mkdir -p Build/Linux
   mkdir -p Lib/Linux/$TARGET_ARCH
   cd Build/Linux
@@ -202,6 +268,7 @@ BuildnvTriStrip()
   echo "building nvTriStrip"
   set -x
   cd Source/ThirdParty/nvTriStrip/nvTriStrip-1.0.0
+  rm -rf Build/Linux
   mkdir -p Build/Linux
   mkdir -p Lib/Linux/$TARGET_ARCH
   cd Build/Linux
@@ -221,6 +288,7 @@ BuildnvTextureTools()
   cd Source/ThirdParty/nvTextureTools/nvTextureTools-2.0.8
   mkdir -p lib/Linux/$TARGET_ARCH
   cd src
+  rm -rf build
   P4Open configure
   chmod +x ./configure
   CXXFLAGS=-fPIC ./configure --release
@@ -251,10 +319,10 @@ BuildSDL2()
   cd $SDL_BUILD_DIR
 
   cmake ../$SDL_DIR
-  make
+  make ${MAKE_ARGS}
   P4Open ../$SDL_DIR/lib/Linux/x86_64-unknown-linux-gnu/libSDL2.a
   P4Open ../$SDL_DIR/lib/Linux/x86_64-unknown-linux-gnu/libSDL2_fPIC.a
-  
+
   cp --remove-destination libSDL2.a ../$SDL_DIR/lib/Linux/x86_64-unknown-linux-gnu/libSDL2.a
   cp --remove-destination libSDL2_fPIC.a ../$SDL_DIR/lib/Linux/x86_64-unknown-linux-gnu/libSDL2_fPIC.a
   set +x
@@ -304,35 +372,150 @@ Run()
 {
   cd ${TOP_DIR}
   echo "==> $1"
-  if [ -n "$VERBOSE" ]; then
+  if [[ $VERBOSE -eq 1 ]]; then
     $1
   else
     $1 >> ${SCRIPT_DIR}/BuildThirdParty.log 2>&1
   fi
 }
 
+Success() {
+    if [ -z $1 ]; then
+        echo "**********  SUCCESS ****************"
+    else
+        echo "**********  SUCCESS $1 ****************"
+    fi
+}
+build_all() {
+   Run BuildJemalloc
+   Run BuildZ
+   Run BuildOpus
+   Run BuildOgg
+   Run BuildVorbis
+   Run BuildHLSLCC
+   Run BuildMcpp
+   Run BuildFreeType
+   Run BuildLND
+   Run BuildForsythTriOO
+   Run BuildnvTriStrip
+   Run BuildnvTextureTools
+   Run BuildSDL2
+   Run Buildcoremod
+   Run FixICU
+}
+
+print_help() {
+
+ echo "-a|--all Rebuild all ThirdParty libs"
+ echo "-v|--verbose Print full output, otherwise output is logged to BuildThirdParty.log"
+ echo "-b|--build BUILD_ME|--build=BUILD_ME Rebuild a specific library"
+ print_valid_build_opts
+ echo "-h|--help Print this help and exit"
+
+}
+
+print_valid_build_opts() {
+  echo "  Valid build choices are"
+  echo "    Jemalloc"
+  echo "    Z"
+  echo "    Opus"
+  echo "    Ogg"
+  echo "    Vorbis"
+  echo "    HLSLCC"
+  echo "    Mcpp"
+  echo "    FreeType"
+  echo "    LND"
+  echo "    ForsythTriOO"
+  echo "    nvTriStrip"
+  echo "    nvTextureTools"
+  echo "    SDL2"
+  echo "    coremod"
+  echo "    ICU"
+}
+
+BuildList=()
+
+# Parse command line parameters
+while :; do
+    case $1 in
+        -h|--help)
+            print_help
+            exit
+            ;;
+        -v|--verbose)
+            VERBOSE=1
+            shift
+            continue
+            ;;
+        -a|--all)
+            ALL=1
+            shift
+            continue
+            ;;
+        -b|--build) # Add to the list of packages to build
+            if [ "$2" ]; then
+                BuildList+=("$2")
+                BUILD_LIST=1
+                shift 2
+                continue
+             else
+                echo 'ERROR: Must specify non empty "--build=BUILD" argument.' >&2
+                print_valid_build_opts
+                exit 1
+             fi
+             ;;
+        --build=?*) # Add to the list of packages to build
+            BuildList+=("${1#*=}")
+            BUILD_LIST=1
+            shift
+            continue
+            ;;
+        --build=)
+            echo 'ERROR: Must specify non empty "--build=BUILD" argument.' >&2
+            exit 1
+            ;;
+        --)
+            shift
+            break
+            ;;
+        -?*)
+            printf "WARN: Unknown Option (ignored): %s\n' "$1" >&2"
+            ;;
+        *)
+            break
+  esac
+
+  shift
+done
+
 echo "Building ThirdParty libraries"
-if [ -z "$VERBOSE" ]; then
-  echo "(For full output set VERBOSE=1, otherwise output is logged to BuildThirdParty.log)"
-fi
 
 rm -f ${SCRIPT_DIR}/BuildThirdParty.log
-if [ -z "$1" ]; then
-  # use bundled Run BuildOpus
-  # use bundled Run BuildOgg
-  # use bundled Run BuildVorbis
-  Run BuildHLSLCC
-  Run BuildMcpp
-  # use bundled Run BuildFreeType
-  Run BuildLND
-  Run BuildForsythTriOO
-  Run BuildnvTriStrip
-  Run BuildnvTextureTools
-  Run BuildSDL2
-  Run Buildcoremod
-  Run FixICU
+
+# if -b --build or --build=?* is used, build the list
+# instead of the default or all
+if [[ ${#BuildList[@]} -ne 0 &&  $ALL -eq 0 ]]; then
+    for var in "${BuildList[@]}"
+    do
+        if [[ "${var}" -ne "ICU" ]]; then
+            Run Fix"${var}"
+            Success "${var}"
+        else
+            Run Build"${var}"
+            Success "${var}"
+        fi
+    done
+    exit
+elif [[ ${#BuildList[@]} -ne 0 ]] && [[ $ALL -eq 1 ]]; then
+    echo "ERROR: Can not build all and individual packages at the same time"
+    exit
+fi
+
+if [[ $ALL -eq 1 ]]; then
+  build_all
 else
-  Run $1
+  Run BuildLND
 fi
 
-echo "********** SUCCESS ****************"
+Success
+
diff --git a/Engine/Build/BatchFiles/Linux/GenerateProjectFiles.sh b/Engine/Build/BatchFiles/Linux/GenerateProjectFiles.sh
--- a/Engine/Build/BatchFiles/Linux/GenerateProjectFiles.sh
+++ b/Engine/Build/BatchFiles/Linux/GenerateProjectFiles.sh
@@ -1,12 +1,29 @@
 #!/bin/bash
 
 SCRIPT_DIR=$(cd "$(dirname "$BASH_SOURCE")" ; pwd)
+# these need to be passed to children scripts
+export ARCHIVE_ROOT=$HOME/Downloads
+export GITHUB_TAG=4.5.0-preview
 
-set -e
+IsGithubBuild()
+{
+  # if p4 is installed, assume building out of perforce repo (no need to download and fix dependencies)
+  # Override this with -git
+  for Arg in $@; do
+    if [ "$Arg" == "-git" ]; then
+      return 0
+    fi
+  done
+
+  if which p4 > /dev/null; then
+    return 1	# perforce
+  fi
+
+  return 0
+}
 
-echo
-echo Setting up Unreal Engine 4 project files...
-echo
+
+set -e
 
 TOP_DIR=$(cd $SCRIPT_DIR/../../.. ; pwd)
 cd ${TOP_DIR}
@@ -17,7 +34,7 @@ located inside the Engine/Build/BatchFiles/Linux directory."
   exit 1
 fi
 
-if [ "$(lsb_release --id)" = "Distributor ID:	Ubuntu" -o "$(lsb_release --id)" = "Distributor ID:	Debian" ]; then
+if [ "$(lsb_release --id)" = "Distributor ID:	Ubuntu" -o "$(lsb_release --id)" = "Distributor ID:	Debian" -o "$(lsb_release --id)" = "Distributor ID:	Linux Mint" ]; then
   # Install all necessary dependencies
   DEPS="mono-xbuild \
     mono-dmcs \
@@ -27,18 +44,56 @@ if [ "$(lsb_release --id)" = "Distributor ID:	Ubuntu" -o "$(lsb_release --id)" =
     libmono-system-management4.0-cil
     libmono-system-xml-linq4.0-cil
     libmono-corlib4.0-cil
-    libogg-dev"
+    libqt4-dev
+    dos2unix
+    "
 
   for DEP in $DEPS; do
     if ! dpkg -s $DEP > /dev/null 2>&1; then
       echo "Attempting installation of missing package: $DEP"
       set -x
-      sudo apt-get install $DEP
+      sudo apt-get install -y $DEP
       set +x
     fi
   done
 fi
 
+echo 
+if IsGithubBuild $@; then
+	echo
+	echo Github build
+	echo Checking / downloading the latest archives
+	echo
+	set +e
+	Build/BatchFiles/Linux/GetAssets.py EpicGames/UnrealEngine $GITHUB_TAG 
+
+	# check if it had to download anything
+	if [ $? -eq 2 ]; then
+	  echo
+          echo Downloaded new binaries!
+	  echo Unpacking and massaging the files
+	  pushd Build/BatchFiles/Linux > /dev/null
+	  ./UpdateDeps.sh 
+	  popd > /dev/null
+        else
+          echo
+          echo All assets are up to date, not unpacking zip files again.
+        fi
+
+else
+	echo Perforce build
+	echo Assuming availability of up to date third-party libraries
+fi
+
+echo
+pushd Build/BatchFiles/Linux > /dev/null
+./BuildThirdParty.sh
+popd > /dev/null
+
+echo
+echo Setting up Unreal Engine 4 project files...
+echo
+
 # args: wrong filename, correct filename
 # expects to be in Engine folder
 CreateLinkIfNoneExists()
diff --git a/Engine/Build/BatchFiles/Linux/GetAssets.py b/Engine/Build/BatchFiles/Linux/GetAssets.py
new file mode 100644
--- /dev/null
+++ b/Engine/Build/BatchFiles/Linux/GetAssets.py
@@ -0,0 +1,124 @@
+#!/usr/bin/env python
+# returns 
+#  0 if assets are up to date
+#  1 if there was any kind of error 
+#  2 if assets were not up to date, but have been successfully downloaded
+import os
+import json
+try:
+	from urllib2 import urlopen, Request
+except:
+	from urllib.request import urlopen, Request
+from sys import argv, stdout
+import time
+
+argv0 = os.path.basename(argv[0])
+
+if len(argv) < 3:
+	print("Usage: {0} <username/repo> <release-tag>".format(argv0))
+	print("")
+	print("example: {0}  EpicGames/UnrealEngine latest-preview".format(argv0))
+	exit(1)
+
+try:
+	token=os.environ["OAUTH_TOKEN"]
+except:
+	print("Please set OAUTH_TOKEN")
+	print("You can generate an access token at [REDACTED_PROFILE_URL]/tokens/new")
+	print("It must have repo scope")
+	exit(1)
+
+try:
+	download_path=os.environ["ARCHIVE_ROOT"]
+except:
+	download_path=os.path.expanduser("~/Downloads")
+
+
+
+repo = argv[1]
+release_name = argv[2]
+
+#only_names = set(argv[3:]) #later...
+
+def get_releases(repo):
+	try:
+		return json.loads(urlopen("https://api.github.com/repos/{0}/releases?access_token={1}".format(repo, token)).read().decode("utf-8"))
+	except:
+		return None
+
+def progress(done, sofar):
+	filled = int((sofar * 72) / done)
+	return "{0} {1} {2:4.1f}%".format("#" * filled, " " * (72-filled),  (float(sofar)/float(done)*100))
+
+def download_assets(repo):
+	releases = get_releases(repo)
+	if not releases:
+		if not repo == "EpicGames/UnrealEngine":
+			return download_assets("EpicGames/UnrealEngine")
+		else:
+			print("Unable to download assets for release {0}!".format(release_name))
+			exit(1)
+
+	had_to_download=False
+	for release in releases:
+		if not release["tag_name"] == release_name:
+			continue
+	
+		if len(release["assets"]) == 0:
+			break
+		for asset in release["assets"]:
+			name = asset["name"]
+			asset_path = os.path.join(download_path, release_name+"-"+name)
+			at_latest_version=False
+	
+			try:
+				with open(asset_path+".updated") as f:
+					timefmt = "%Y-%m-%dT%H:%M:%SZ"
+					downloaded = time.strptime(f.read(), timefmt)
+					updated = time.strptime(asset['updated_at'], timefmt)
+					at_latest_version = downloaded >= updated
+			except:
+				pass
+	
+			size = int(asset["size"])
+			if at_latest_version:
+				if size != os.path.getsize(asset_path):
+					print("Incomplete download of asset {0}, retrying")
+				else:
+					print("Asset {0} up to date".format(name))
+					continue
+	
+			had_to_download=True
+
+			req = Request(asset["url"]+"?access_token="+token, None, {"Accept": "application/octet-stream"})
+			handle = urlopen(req)
+			length = int(handle.headers["Content-Length"])
+	
+			print("Started downloading {0}".format(name))
+	
+			with open(asset_path, "wb") as f:
+				read = 0
+				while True:
+					chunk = handle.read(4096)
+					if len(chunk):
+						f.write(chunk)
+					else:
+						break
+					read += len(chunk)
+					stdout.write("\r" + progress(length, read))
+					stdout.flush()
+			with open(asset_path+".updated", 'w') as f:
+				f.write(asset['updated_at'])
+			print("\nFinished downloading {0}".format(name))
+	
+		if had_to_download:
+			return 2
+		return 0
+	if not repo == "EpicGames/UnrealEngine":
+		return download_assets("EpicGames/UnrealEngine")
+
+	print("Unable to download assets for release {0}!".format(release_name))
+	exit(1)
+
+exit(download_assets(repo))
+
diff --git a/Engine/Build/BatchFiles/Linux/UpdateDeps.sh b/Engine/Build/BatchFiles/Linux/UpdateDeps.sh
--- a/Engine/Build/BatchFiles/Linux/UpdateDeps.sh
+++ b/Engine/Build/BatchFiles/Linux/UpdateDeps.sh
@@ -1,27 +1,31 @@
 #!/bin/bash
-# This script will verify the downloaded zip file dependencies, unzip them and
-# apply linux-specific patches.  It tries to clean up any previous versions
-# using 'git clean' on certain directories.
+# This script will unzip the downloaded archives and convert to Unix format where needed.  
+# It tries to clean up any previous versions using 'git clean' on certain directories.
 # WARNING: This script can be destructive so use with caution if you have
 # unsaved local changes.
 #
 # This script does not handle building of any of the ThirdParty dependencies.
 # (See BuildThirdParty.sh).
 
-# Location of downloaded .zip files.
-if [ -n "$1" ]; then
-  ARCHIVE_ROOT=$1
-else
-  ARCHIVE_ROOT=$HOME/Downloads
+# Location of the archive files (shared with GetAssets.py)
+if [ -z $ARCHIVE_ROOT ]; then
+  echo You should have ARCHIVE_ROOT variable set to the same location as used in GetAssets.py. If you don't know what this is about, don't invoke this file directly.
 fi
 
+# Prefix of the archive files (should match one used by GetAssets.py)
+if [ -z $GITHUB_TAG ]; then
+  echo You should have GITHUB_TAG variable set to the same release tag as used in GetAssets.py. If you don't know what this is about, don't invoke this file directly.
+fi
 
 if [ ! -d $ARCHIVE_ROOT ]; then
   echo "Download root not found: $ARCHIVE_ROOT"
   exit 1
 fi
 
-# Zip file dependencies.
+# Zipped archives.
+ARCHIVES="$GITHUB_TAG-Required_1of2.zip
+    $GITHUB_TAG-Required_2of2.zip
+    $GITHUB_TAG-Optional.zip"
 
 # Make sure that the current working directory is the root the git checkout,
 # which is four levels up from the this script.
@@ -31,88 +35,8 @@ cd ${TOP_DIR}
 set -e
 trap "echo '==> An error occured while running UpdateDeps.sh'" ERR
 
-DEPENDENCIES=$SCRIPT_DIR/dependencies.txt
-INSTALLED_DEPENDENCIES=${DEPENDENCIES%.txt}.installed
-
-#TODO: Move these hardcodings out of this file
-PERFORCE_ARCHIVE=$ARCHIVE_ROOT/p4api.tgz
-STEAM_ARCHIVE=$ARCHIVE_ROOT/steamworks_sdk_130.zip
-
-# Make it super easy for users of certain linux distributions.
-CheckDistroDependencies() {
-  echo "==> Checking build dependencies"
-  if [ "$(lsb_release --id)" = "Distributor ID:	Ubuntu" -o "$(lsb_release --id)" = "Distributor ID:	Debian" ]; then
-    # Install all necessary dependencies
-    DEPS="libqt4-dev
-      libglew-dev
-      dos2unix 
-      clang-3.3"
-
-    for DEP in $DEPS; do
-      if ! dpkg -s $DEP > /dev/null 2>&1; then
-        echo "Attempting installation of missing package: $DEP"
-        set -x
-        sudo apt-get install $DEP
-        set +x
-      fi
-    done
-  fi
-}
-
-CheckInstalled() {
-  echo "==> Checking existing install"
-  # Phase one, check sha1 sums of previously installed files.
-  if [ -f $INSTALLED_DEPENDENCIES ]; then
-    if [ "`cat $INSTALLED_DEPENDENCIES`" = "`cat $DEPENDENCIES`" ]; then
-      echo "Dependencies are up-to-date."
-      echo "(remove $INSTALLED_DEPENDENCIES file and re-run to force install)."
-      exit 0
-    else
-      echo "Installed dependencies are out-of-date."
-    fi
-  else
-    echo "No installed dependenies found."
-  fi
-}
-
-Download() {
-  # For those archives that are missing some can be downloaded automatically
-  # using curl.
-  cat $DEPENDENCIES | while read line; do
-    DEP_URL=`echo $line | cut -d " " -f 2`
-    DEP_WALLED=`echo $line | cut -d " " -f 3`
-    DEP_FILE=$(basename $DEP_URL)
-    if [ ! -f $ARCHIVE_ROOT/$DEP_FILE ]; then
-      if [ "$DEP_WALLED" -eq 0 ]; then
-        echo "==> Downloading $DEP_FILE"
-        curl --progress -o $ARCHIVE_ROOT/$DEP_FILE $DEP_URL
-      fi
-    fi
-  done
-}
-
-VerifyDownloads() {
-  echo "==> Verifying downloaded archives"
-  cd $ARCHIVE_ROOT
-  cat $DEPENDENCIES | while read line; do
-    DEP_HASH=`echo $line | cut -d " " -f 1`
-    DEP_URL=`echo $line | cut -d " " -f 2`
-    DEP_FILE=$(basename $DEP_URL)
-    if ! echo "$DEP_HASH  $DEP_FILE" | sha1sum -c -; then
-      echo ""
-      echo "ERROR: Missing or out-of-date dependency: $DEP_FILE"
-      echo "Please download the following file to $ARCHIVE_ROOT:"
-      echo "  $DEP_URL"
-      echo "If you have downloaded these files somewhere else you can"
-      echo "pass the location as the first argument to this script."
-      exit 1
-    fi
-  done
-
-  cd ${TOP_DIR}
-}
-
-ExtractArchives() {
+CleanRepo() 
+{
   # Clean any files installed from .zip files.
   echo "==> Cleaning existing dependencies"
   echo "About to cleanup previous zip contents with git clean."
@@ -137,30 +61,28 @@ ExtractArchives() {
     Engine/Extras
     Samples/StarterContent
     Templates"
-  for dir in $CLEAN_DIRS; do
-    git clean -xfd $dir
+  for Dir in $CLEAN_DIRS; do
+    git clean -xfd $Dir
   done
+}
 
+ExtractArchives()
+{
   echo "==> Installing dependencies"
-  # Extract files and write new sha1 files.
-  cat $DEPENDENCIES | while read line; do
-    DEP_URL=`echo $line | cut -d " " -f 2`
-    if [[ $(dirname $DEP_URL) != https://github.com/EpicGames/* ]]; then
-      continue
-    fi
-    DEP_FILE=$(basename $DEP_URL)
-    file_full=$ARCHIVE_ROOT/$DEP_FILE
-    if [ ! -f $file_full ]; then
-      echo "Zip file dependency not found: $file_full"
+  for Archive in $ARCHIVES; do
+    FullPathArchive=$ARCHIVE_ROOT/$Archive
+    if [ ! -f $FullPathArchive ]; then
+      echo "Zip file dependency not found: $FullPathArchive"
       exit 1
     fi
-    echo "Extracting $file_full"
-    unzip -q -n $file_full
+    echo "Extracting $FullPathArchive"
+    unzip -q -n $FullPathArchive
   done
 }
 
-ApplyPatches() {
-  echo "==> Applying patches"
+Unixify()
+{
+  echo "==> Converting to unix line endings"
   cd Engine/Source/ThirdParty
   echo "PWD=$PWD"
 
@@ -173,7 +95,6 @@ ApplyPatches() {
   dos2unix configure.ac
   dos2unix `find -name "*.in"`
   dos2unix `find -name "*.m4"`
-  autoconf
   cd -
 
   cd MCPP/mcpp-2.7.2/
@@ -195,9 +116,6 @@ ApplyPatches() {
   dos2unix depcomp
   dos2unix config.*
   dos2unix `find -name Makefile.in`
-  chmod +x configure
-  chmod +x config.sub
-  chmod +x config.guess
   cd -
 
   cd Ogg/libogg-1.2.2/
@@ -207,81 +125,14 @@ ApplyPatches() {
   dos2unix depcomp
   dos2unix config.*
   dos2unix `find -name Makefile.in`
-  chmod +x configure
-  chmod +x config.sub
-  chmod +x config.guess
-  cd -
-
-  cd SDL2/
-  pwd
-  tar xf SDL2-2.0.3/build/SDL2-2.0.3.tar.gz -C SDL2-2.0.3/build
-  rm -f DisplayServerExtensions/bin/libdsext.a
-  patch -f -p1 < ${SCRIPT_DIR}/SDL2.patch
-  cd -
-
-  # This needs to be moved out of the way as the LinuxToolChain.cs
-  # Links via -l and not via path name so will pick the .so instead
-  cd FBX/2014.2.1/
-  pwd
-  rm -f lib/linux/x86_64-unknown-linux-gnu/libfbxsdk.so
   cd -
 
   cd ${TOP_DIR}
 }
 
-RenameContent() {
-  echo "==> Correcting Filename case"
-  cd Engine/Content
-  echo "PWD=$PWD"
-
-  find . -name "*.PNG" -exec rename 's/\.PNG$/\.png/' {} \;
-
-  cd ${TOP_DIR}
-}
-
-InstallPerforce() {
-  echo "==> Installing Perforce (p4api) libraries"
-  cd Engine/Source/ThirdParty/Perforce
-  rm -rf p4api-2012.1.442152 p4api-2012.1/lib/linux
-  tar xf $PERFORCE_ARCHIVE
-  cp -r p4api-2012.1.442152/lib p4api-2012.1/lib/linux
-  cd ${TOP_DIR}
-}
-
-InstallSteamworks() {
-  echo "==> Installing Steamworks SDK"
-  rm -rf Engine/Source/ThirdParty/Steamworks/Steamv130
-  mkdir Engine/Source/ThirdParty/Steamworks/Steamv130
-  cd Engine/Source/ThirdParty/Steamworks/Steamv130
-  unzip -q $STEAM_ARCHIVE
-  cd ${TOP_DIR}
-  mkdir -p Engine/Binaries/Linux/
-  rm -f Engine/Binaries/Linux/libsteam_api.so
-  cd Engine/Binaries/Linux
-  ln -s ../../Source/ThirdParty/Steamworks/Steamv130/sdk/redistributable_bin/linux64/libsteam_api.so
-  cd ${TOP_DIR}
-}
-
-UpdateInstallStamp() {
-  echo "==> Updating install stamp ($INSTALLED_DEPENDENCIES)"
-  cp $DEPENDENCIES $INSTALLED_DEPENDENCIES
-}
-
-CheckDistroDependencies
-CheckInstalled
-Download
-VerifyDownloads
+CleanRepo
 ExtractArchives
-ApplyPatches
-RenameContent
-InstallPerforce
-InstallSteamworks
-UpdateInstallStamp
+Unixify
 
 echo "********** SUCCESS ****************"
 
-# Once we've successfully updated the deps we need to rebuild the
-# ThirdParty modules.  This is mostly so that we get everything built
-# with -fPIC.  BuildThirdParty.sh is a seperate script so that it can
-# be run repeatedly and independently of UpdateDeps.sh if needed.
-# $SCRIPT_DIR/BuildThirdParty.sh
