Fixes to native Linux build scripts.

diff --git a/Engine/Build/BatchFiles/Linux/GenerateProjectFiles.sh b/Engine/Build/BatchFiles/Linux/GenerateProjectFiles.sh
--- a/Engine/Build/BatchFiles/Linux/GenerateProjectFiles.sh
+++ b/Engine/Build/BatchFiles/Linux/GenerateProjectFiles.sh
@@ -66,9 +66,17 @@ if IsGithubBuild $@; then
 	echo
 	set +e
 	Build/BatchFiles/Linux/GetAssets.py EpicGames/UnrealEngine $GITHUB_TAG 
+	UpdateResult=$?
+	set -e
+
+	if [ $UpdateResult -eq 1 ]; then
+	  echo
+          echo Could not check/download binary dependencies!
+          exit 1
+        fi
 
 	# check if it had to download anything
-	if [ $? -eq 2 ]; then
+	if [ $UpdateResult -eq 2 ]; then
 	  echo
           echo Downloaded new binaries!
 	  echo Unpacking and massaging the files
diff --git a/Engine/Build/BatchFiles/Linux/GetAssets.py b/Engine/Build/BatchFiles/Linux/GetAssets.py
--- a/Engine/Build/BatchFiles/Linux/GetAssets.py
+++ b/Engine/Build/BatchFiles/Linux/GetAssets.py
@@ -83,7 +83,7 @@ def download_assets(repo):
 			size = int(asset["size"])
 			if at_latest_version:
 				if size != os.path.getsize(asset_path):
-					print("Incomplete download of asset {0}, retrying")
+					print("Incomplete download of asset {0}, retrying".format(name))
 				else:
 					print("Asset {0} up to date".format(name))
 					continue
