Merging back updated Linux readme from 4.5 to main.

diff --git a/Engine/Build/BatchFiles/Linux/README.md b/Engine/Build/BatchFiles/Linux/README.md
--- a/Engine/Build/BatchFiles/Linux/README.md
+++ b/Engine/Build/BatchFiles/Linux/README.md
@@ -1,49 +1,105 @@
 Build scripts for native Linux build
 ====================================
 
-This document describes how to build UE4 natively on a Linux host.
-So far this has only been tested on Debian/Sid and Ubuntu/Precise.
+This document describes how to build UE4 4.5 natively on a Linux host.
+The steps are described here are applicable to the current build, but you may
+want to visit https://wiki.unrealengine.com/Building_On_Linux for the
+latest updates on the process.
+
+If you are stuck at some point, we suggest searching AnswerHub 
+(https://answers.unrealengine.com/questions/topics/linux.html) for possible answers 
+or asking a new question. You may also receive help on #UE4Linux IRC channel on FreeNode,
+however it is not an official support outlet.
 
-A wiki page page regarding building on linux is also maintained at
-https://wiki.unrealengine.com/Building_On_Linux, and public IRC discussion
-happens on #UE4Linux on FreeNode.
 
 Prerequisites
 -------------
 
-You will need mono + gmcs and several other packages installed in order to build
-EU4. The command to install these on Ubundu 12.04 (Precise) is:
+Packages required to build the engine vary from distribution to distribution,
+and an up-to-date list should be maintained (and installed) by GenerateProjectFiles.sh
+- feel free to suggest modifications. Automated install currently works for Debian-based
+distributions only (Debian itself, (K)Ubuntu and Linux Mint).
+
+Most important dependencies:
+- mono and 4.0 (at least) libraries, including xbuild and C# compiler (*mcs).
+- clang 3.3 (clang 3.5.0 is also supported, but NOT clang 3.4).
+- python (2 or 3) - needed for the script that downloads the binary dependencies.
+- Qt(4,5) or GTK development packages to build LinuxNativeDialogs.
+- SDL2 is also needed for building LinuxNativeDialogs module, but the rest
+of the engine is using our own (modified) version of it from Engine/Source/ThirdParty/SDL2.
+
+You will also need at least 20 GB of disk space and a relatively powerful
+machine.
+
+If you want to rebuild third-party dependencies (we don't recommend doing
+that anymore), you will need many more development packages installed. Refer
+to BuildThirdParty.sh script and specific automake/CMake scripts for each
+dependency. You don't have to do that though as we supply prebuilt libraries.
 
-    $ apt-get install mono-mcs mono-xbuild dos2unix clang-3.3 \
-        libogg-dev libmono-microsoft-build-tasks-v4.0-4.0-cil \
-        libmono-system-data-datasetextensions4.0-cil \
-        libmono-system-web-extensions4.0-cil \
-        libmono-system-management4.0-cil
 
 Building
 --------
 
-1. Download the additional .zip archives for the master branch from Epic's GitHub
-   Releases page at https://github.com/EpicGames/UnrealEngine/releases
+Build steps have been simplified since previous releases, and the additional
+.zip archives are now being downloaded automatically. However, in order to
+make this happen, you need to generate a personal access token (repo scope) here:
+[REDACTED_PROFILE_URL]/tokens/new   
+Copy the token to OAUTH_TOKEN environment variable - treat it as your GitHub password
+in terms of security.
+More here: https://help.github.com/articles/creating-an-access-token-for-command-line-use/
+
+Step by step:
+
+1. Clone EpicGames/UnrealEngine repository
+
+    git clone https://github.com/EpicGames/UnrealEngine -b 4.5
+
+2. With the aforementioned OAUTH_TOKEN variable set to a valid token,
+run GenerateProjectFiles.sh file in the engine's directory:
+
+   cd UnrealEngine
+   ./GenerateProjectFiles.sh
+
+The script may ask you to install additional packages (on certain distributions), and then
+it will proceed to downloading the archives with binary dependencies (by default, to ~/Downloads), 
+which are pretty large (3GB in total) - they won't be re-downloaded unless they are updated 
+for the particular release tag. If new archives have been downloaded (which is the case the
+first time you do this), the script will unpack them over your repository (it will ask you 
+to confirm that action as it is potentially destructive).
+After that, it will build LinuxNativeDialogs - the only third-party library that currently 
+needs to be built locally.
+
+3. After successful execution of GenerateProjectFiles.sh, you can build the editor:
+
+   make ShaderCompileWorker UnrealLightmass UE4Editor
+
+That's it.
 
-2. Download v130 of the Steamworks SDK.
 
-3. Install the downloaded dependencies with:
+Notes
+-----
 
-        $ ./UpdateDeps.sh.
+GenerateProjectFiles.sh also produces CMakeLists.txt which you can use to import the
+project in your favorite IDE. Caveat: only KDevelop 4.6+ is known to handle the project
+well (it takes about 3 GB of resident RAM to load the project in KDevelop 4.7).
 
-4. Build third party libraries:
+If you intend to develop the editor, you can build a debug configuration of it:
 
-        $ ./BuildThirdParty.sh.
+  make UE4Editor-Linux-Debug
 
-5. Build UnrealBuildTool (UBT) and generate top level Makefile:
+(note that it will still use development ShaderCompileWorker / UnrealLightmass). This
+configuration runs much slower.
 
-        $ ./GenerateProjectFiles.sh
+If you want to rebuild the editor from scatch, you can use
 
-6. Build your targets using the top-level Makefile. e.g:
+  make UE4Editor ARGS="-clean" && make UE4Editor
 
-        $ make UE4Client
+The time it takes to build the editor in development configuration can be large,
+debug configuration takes about 2/3 of this time. The build process can also take 
+significant amount of RAM (roughly 1GB per core).
 
-   Or by running `Build.sh` directly:
+It is also possible to cross-compile the editor (currently from Windows only), 
+although you will still need to build LinuxNativeDialogs on the Linux machine. 
+You may use this route if your Windows machine happens to be more powerful,
+but explanation of it is beyond the scope of this document.
 
-        $ Engine/Build/BatchFiles/Linux/Build.sh UE4Client Linux Debug
