[INTEGRATE] Content back from 4.5 branch.

diff --git a/Engine/Build/BatchFiles/Linux/GenerateProjectFiles.sh b/Engine/Build/BatchFiles/Linux/GenerateProjectFiles.sh
--- a/Engine/Build/BatchFiles/Linux/GenerateProjectFiles.sh
+++ b/Engine/Build/BatchFiles/Linux/GenerateProjectFiles.sh
@@ -5,21 +5,24 @@ SCRIPT_DIR=$(cd "$(dirname "$BASH_SOURCE")" ; pwd)
 export ARCHIVE_ROOT=$HOME/Downloads
 export GITHUB_TAG=4.5.0-preview
 
-IsGithubBuild()
+IS_GITHUB_BUILD=true
+FORCE_UPDATEDEPS=false
+
+CheckArgs()
 {
   # if p4 is installed, assume building out of perforce repo (no need to download and fix dependencies)
+  if which p4 > /dev/null; then
+    IS_GITHUB_BUILD=false	# perforce
+  fi
+
   # Override this with -git
   for Arg in $@; do
     if [ "$Arg" == "-git" ]; then
-      return 0
+      IS_GITHUB_BUILD=true
+    elif [ "$Arg" == "-updatedeps" ]; then
+      FORCE_UPDATEDEPS=true     
     fi
   done
-
-  if which p4 > /dev/null; then
-    return 1	# perforce
-  fi
-
-  return 0
 }
 
 
@@ -34,6 +37,8 @@ located inside the Engine/Build/BatchFiles/Linux directory."
   exit 1
 fi
 
+CheckArgs $@
+
 if [ "$(lsb_release --id)" = "Distributor ID:	Ubuntu" -o "$(lsb_release --id)" = "Distributor ID:	Debian" -o "$(lsb_release --id)" = "Distributor ID:	Linux Mint" ]; then
   # Install all necessary dependencies
   DEPS="mono-xbuild \
@@ -46,6 +51,7 @@ if [ "$(lsb_release --id)" = "Distributor ID:	Ubuntu" -o "$(lsb_release --id)" =
     libmono-corlib4.0-cil
     libqt4-dev
     dos2unix
+    cmake
     "
 
   for DEP in $DEPS; do
@@ -59,12 +65,13 @@ if [ "$(lsb_release --id)" = "Distributor ID:	Ubuntu" -o "$(lsb_release --id)" =
 fi
 
 echo 
-if IsGithubBuild $@; then
+if [ "$IS_GITHUB_BUILD" = true ]; then
 	echo
 	echo Github build
 	echo Checking / downloading the latest archives
 	echo
 	set +e
+        chmod +x Build/BatchFiles/Linux/GetAssets.py
 	Build/BatchFiles/Linux/GetAssets.py EpicGames/UnrealEngine $GITHUB_TAG 
 	UpdateResult=$?
 	set -e
@@ -76,7 +83,7 @@ if IsGithubBuild $@; then
         fi
 
 	# check if it had to download anything
-	if [ $UpdateResult -eq 2 ]; then
+	if [ $UpdateResult -eq 2 -o "$FORCE_UPDATEDEPS" = true ]; then
 	  echo
           echo Downloaded new binaries!
 	  echo Unpacking and massaging the files
@@ -109,12 +116,12 @@ CreateLinkIfNoneExists()
     WrongName=$1
     CorrectName=$2
 
-    pushd `dirname $CorrectName`
+    pushd `dirname $CorrectName` > /dev/null
     if [ ! -f `basename $CorrectName` ] && [ -f $WrongName ]; then
       echo "$WrongName -> $CorrectName"
       ln -sf $WrongName `basename $CorrectName`
     fi
-    popd
+    popd > /dev/null
 }
 
 
diff --git a/Engine/Build/BatchFiles/Linux/README.md b/Engine/Build/BatchFiles/Linux/README.md
--- a/Engine/Build/BatchFiles/Linux/README.md
+++ b/Engine/Build/BatchFiles/Linux/README.md
@@ -1,98 +1,135 @@
 Build scripts for native Linux build
 ====================================
 
-This document describes how to build UE4 4.5 natively on a Linux host.
+This document describes how to build UE4 4.5 natively on a Linux host. 
 The steps are described here are applicable to the current build, but you may
 want to visit https://wiki.unrealengine.com/Building_On_Linux for the
 latest updates on the process.
 
 If you are stuck at some point, we suggest searching AnswerHub 
 (https://answers.unrealengine.com/questions/topics/linux.html) for possible answers 
-or asking a new question. You may also receive help on #UE4Linux IRC channel on FreeNode,
-however it is not an official support outlet.
+or asking a new question on there if you can not find what you are looking for. 
+You may also receive help on #UE4Linux IRC channel on FreeNode, however it is 
+not an official support outlet.
 
 
 Prerequisites
 -------------
 
-Packages required to build the engine vary from distribution to distribution,
-and an up-to-date list should be maintained (and installed) by GenerateProjectFiles.sh
-- feel free to suggest modifications. Automated install currently works for Debian-based
+The packages that are required to build the engine vary from distribution to distribution,
+and an up-to-date list should be maintained (and installed) by GenerateProjectFiles.sh -
+feel free to suggest modifications. Automated install currently works for Debian-based
 distributions only (Debian itself, (K)Ubuntu and Linux Mint).
 
 Most important dependencies:
-- mono and 4.0 (at least) libraries, including xbuild and C# compiler (*mcs).
+- mono 3.x (2.x may work, but is not recommended), including xbuild and C# compiler (*mcs), and libraries for NET 4.0 framework.
 - clang 3.3 (clang 3.5.0 is also supported, but NOT clang 3.4).
 - python (2 or 3) - needed for the script that downloads the binary dependencies.
 - Qt(4,5) or GTK development packages to build LinuxNativeDialogs.
 - SDL2 is also needed for building LinuxNativeDialogs module, but the rest
 of the engine is using our own (modified) version of it from Engine/Source/ThirdParty/SDL2.
 
-You will also need at least 20 GB of disk space and a relatively powerful
+You will also need at least 20 GB of free disk space and a relatively powerful
 machine.
 
 If you want to rebuild third-party dependencies (we don't recommend doing
-that anymore), you will need many more development packages installed. Refer
+that any more), you will need many more development packages installed. Refer
 to BuildThirdParty.sh script and specific automake/CMake scripts for each
 dependency. You don't have to do that though as we supply prebuilt libraries.
 
 
-Building
---------
+Setting up/updating the sources
+-------------------------------
 
-Build steps have been simplified since previous releases, and the additional
-.zip archives are now being downloaded automatically. However, in order to
-make this happen, you need to generate a personal access token (repo scope) here:
-[REDACTED_PROFILE_URL]/tokens/new   
-Copy the token to OAUTH_TOKEN environment variable - treat it as your GitHub password
-in terms of security.
-More here: https://help.github.com/articles/creating-an-access-token-for-command-line-use/
+Setup has been simplified since the previous releases, and the additional
+.zip archives are now being automatically downloaded. However, in order to
+make this happen, you need to generate a personal access token (repo scope):
+[REDACTED_PROFILE_URL]/tokens/new  Copy the token to OAUTH_TOKEN environment variable -\
+treat it as your GitHub password in terms of security and have that variable set each
+time you are running the below script.
+
+More information on github access tokens: https://help.github.com/articles/creating-an-access-token-for-command-line-use/
 
 Step by step:
 
 1. Clone EpicGames/UnrealEngine repository
 
+
     git clone https://github.com/EpicGames/UnrealEngine -b 4.5
 
 2. With the aforementioned OAUTH_TOKEN variable set to a valid token,
 run GenerateProjectFiles.sh file in the engine's directory:
 
-   cd UnrealEngine
-   ./GenerateProjectFiles.sh
 
-The script may ask you to install additional packages (on certain distributions), and then
-it will proceed to downloading the archives with binary dependencies (by default, to ~/Downloads), 
-which are pretty large (3GB in total) - they won't be re-downloaded unless they are updated 
-for the particular release tag. If new archives have been downloaded (which is the case the
-first time you do this), the script will unpack them over your repository (it will ask you 
-to confirm that action as it is potentially destructive).
-After that, it will build LinuxNativeDialogs - the only third-party library that currently 
-needs to be built locally.
+    cd UnrealEngine
+    ./GenerateProjectFiles.sh
 
-3. After successful execution of GenerateProjectFiles.sh, you can build the editor:
 
-   make ShaderCompileWorker UnrealLightmass UE4Editor
+The script may ask you to install additional packages (on certain distributions), then it will download 
+the archives with binary dependencies (by default, to ~/Downloads), which is pretty large (3GB in total). 
+They won't be re-downloaded unless they are updated for the particular release tag. 
+If any new archives have been downloaded (which is the case the first time you do this), 
+the script will unpack them over your repository (it will ask you 
+to confirm that action as it is potentially destructive). You can force this by supplying 
+-updatedeps switch if for some reason you want to unpack them again.
 
-That's it.
+The script will also build LinuxNativeDialogs - the only third-party library that currently 
+needs to be built locally.
 
+Subsequent updates can be done the same way (remember about OAUTH_TOKEN!)
 
-Notes
------
+    cd UnrealEngine
+    git pull
+    ./GenerateProjectFiles.sh
+
+
+Building and running
+--------------------
 
-GenerateProjectFiles.sh also produces CMakeLists.txt which you can use to import the
+GenerateProjectFiles.sh also produces both makefile and CMakeLists.txt which you can use to import the
 project in your favorite IDE. Caveat: only KDevelop 4.6+ is known to handle the project
-well (it takes about 3 GB of resident RAM to load the project in KDevelop 4.7).
+well (it takes about 3-4 GB of resident RAM to load the project).
+
+The targets match the name of the resulting binary, e.g. UE4Editor-Linux-Debug or UE4Game. You can build them
+by just typing make <target> in the engine's root folder.
+
+Specifically, to be able to run the editor, build the following targets:
+
+    make ShaderCompileWorker UnrealLightmass UE4Editor
 
 If you intend to develop the editor, you can build a debug configuration of it:
 
-  make UE4Editor-Linux-Debug
+    make UE4Editor-Linux-Debug
 
 (note that it will still use development ShaderCompileWorker / UnrealLightmass). This
 configuration runs much slower.
 
 If you want to rebuild the editor from scatch, you can use
 
-  make UE4Editor ARGS="-clean" && make UE4Editor
+    make UE4Editor ARGS="-clean" && make UE4Editor
+
+In order to run it:
+
+    cd Engine/Binaries/Linux/
+    ./UE4Editor
+
+Or, if you want to start it with a specific project:
+
+    cd Engine/Binaries/Linux/
+    ./UE4Editor "~/Documents/Unreal Projects/MyProject/MyProject.uproject"
+    
+You can also append -game if you want to run the project as a game (you can also do that from the running editor).
+
+Notes
+-----
+
+On the first start, the editor will be "compiling shaders" (a bit misleading terminology, 
+it will be converting them to GLSL). The result will be stored in Engine/DerivedDataCache folder 
+and used for subsequent runs.
+
+Depending on the project, the editor may need rather large number of file handles 
+(e.g. 16000+). If you start seeing errors about not being able to open files, 
+you may need to adjust your limits.
 
 The time it takes to build the editor in development configuration can be large,
 debug configuration takes about 2/3 of this time. The build process can also take 
diff --git a/Engine/Config/BaseEditorGameAgnostic.ini b/Engine/Config/BaseEditorGameAgnostic.ini
--- a/Engine/Config/BaseEditorGameAgnostic.ini
+++ b/Engine/Config/BaseEditorGameAgnostic.ini
@@ -19,7 +19,6 @@ bCaptureAudioFromMicrophone=true
 
 [/Script/IntroTutorials.EditorTutorialSettings]
 +Categories=(Identifier="Basics",Title="Basics",Description="Getting started with Unreal Engine 4.",Icon="PlayWorld.RepeatLastPlay",Texture=None)
-+Categories=(Identifier="Installation",Title="Installation",Description="Installing extra 3rd party software needed, or just useful to have, while working with Unreal Engine 4.",Icon="AssetEditor.SaveAsset",Texture=None)
 +Categories=(Identifier="Mobile",Title="Mobile",Description="Mobile Tutorials.",Icon="MaterialEditor.ToggleMobileStats",Texture=None)
 +Categories=(Identifier="Blueprints",Title="Blueprints",Description="Tutorials covering the creation and usage of Blueprints.",Icon="FullBlueprintEditor.EditGlobalOptions",Texture=/Engine/Tutorial/BlueprintTutorials/TutorialAssets/Blueprint_64x.Blueprint_64x)
 +Categories=(Identifier="Paper2D",Title="Paper2D",Description="Tutorials convering the usage of Unreal Engine 4\'s 2D game tools: Paper2D.",Icon="ClassThumbnail.PaperSprite",Texture=/Engine/Tutorial/Paper2D/TutorialAssets/Paper2DSprite_TutorialIcon.Paper2DSprite_TutorialIcon)
