Made Mac Buildsh always return 0 to solve a problem with hot-reload reporting error in Xcode due to not recognizing 254 (target up to date) as a success exit code.

diff --git a/Engine/Build/BatchFiles/Mac/Build.sh b/Engine/Build/BatchFiles/Mac/Build.sh
--- a/Engine/Build/BatchFiles/Mac/Build.sh
+++ b/Engine/Build/BatchFiles/Mac/Build.sh
@@ -87,5 +87,5 @@ case $ACTION in
 		;;
 esac
 
-exit $?
-
+# UBT can exit with non-zero code that means success, for example 254 when target was up to date. Xcode treats that as error, so we return 0 here instead of $?
+exit 0
