Better error handling in Mac Build.sh

diff --git a/Engine/Build/BatchFiles/Mac/Build.sh b/Engine/Build/BatchFiles/Mac/Build.sh
--- a/Engine/Build/BatchFiles/Mac/Build.sh
+++ b/Engine/Build/BatchFiles/Mac/Build.sh
@@ -87,5 +87,9 @@ case $ACTION in
 		;;
 esac
 
-# UBT can exit with non-zero code that means success, for example 254 when target was up to date. Xcode treats that as error, so we return 0 here instead of $?
-exit 0
+ExitCode=$?
+if [ $ExitCode -eq 254 ] || [ $ExitCode -eq 255 ]; then
+	exit 0
+else
+	exit $ExitCode
+fi
