Point main branch scripts to latest-preview tag.

diff --git a/Engine/Build/BatchFiles/Linux/GenerateProjectFiles.sh b/Engine/Build/BatchFiles/Linux/GenerateProjectFiles.sh
--- a/Engine/Build/BatchFiles/Linux/GenerateProjectFiles.sh
+++ b/Engine/Build/BatchFiles/Linux/GenerateProjectFiles.sh
@@ -3,7 +3,7 @@
 SCRIPT_DIR=$(cd "$(dirname "$BASH_SOURCE")" ; pwd)
 # these need to be passed to children scripts
 export ARCHIVE_ROOT=$HOME/Downloads
-export GITHUB_TAG=4.5.0-preview
+export GITHUB_TAG=latest-preview
 
 IS_GITHUB_BUILD=true
 FORCE_UPDATEDEPS=false
