Changes to TPS unpacker from PR #560

- Contributed by [REDACTED_PERSON].

diff --git a/Engine/Build/BatchFiles/Linux/UpdateDeps.sh b/Engine/Build/BatchFiles/Linux/UpdateDeps.sh
--- a/Engine/Build/BatchFiles/Linux/UpdateDeps.sh
+++ b/Engine/Build/BatchFiles/Linux/UpdateDeps.sh
@@ -23,8 +23,9 @@ if [ ! -d $ARCHIVE_ROOT ]; then
 fi
 
 # Zipped archives.
-ARCHIVES="$GITHUB_TAG-Required_1of2.zip
-    $GITHUB_TAG-Required_2of2.zip
+ARCHIVES="$GITHUB_TAG-Required_1of3.zip
+    $GITHUB_TAG-Required_2of3.zip
+    $GITHUB_TAG-Required_3of3.zip
     $GITHUB_TAG-Optional.zip"
 
 # Make sure that the current working directory is the root the git checkout,
