Add Mono lib directory to LD_LIBRARY_PATH. Fixes issue running GitDependencies.exe.

diff --git a/Engine/Build/BatchFiles/Mac/SetupMono.sh b/Engine/Build/BatchFiles/Mac/SetupMono.sh
--- a/Engine/Build/BatchFiles/Mac/SetupMono.sh
+++ b/Engine/Build/BatchFiles/Mac/SetupMono.sh
@@ -27,6 +27,7 @@ if [ $IS_MONO_INSTALLED -eq 0 ]; then
 	export UE_MONO_DIR=$CUR_DIR/../../../Binaries/ThirdParty/Mono/Mac
 	export PATH=$UE_MONO_DIR/bin:$PATH
 	export MONO_PATH=$UE_MONO_DIR/lib:$MONO_PATH
+	export LD_LIBRARY_PATH=$UE_MONO_DIR/lib:$LD_LIBRARY_PATH
 fi
 
 cd "$START_DIR"
