Add GitDependencies as a post-merge hook.

- Also add it to GenerateProjectFiles.sh so the flow is not too different between 4.5 and 4.6
- Remove unpacking zips and older solution.

diff --git a/Engine/Build/BatchFiles/Linux/GenerateProjectFiles.sh b/Engine/Build/BatchFiles/Linux/GenerateProjectFiles.sh
--- a/Engine/Build/BatchFiles/Linux/GenerateProjectFiles.sh
+++ b/Engine/Build/BatchFiles/Linux/GenerateProjectFiles.sh
@@ -71,8 +71,8 @@ if [ "$IS_GITHUB_BUILD" = true ]; then
 	echo Checking / downloading the latest archives
 	echo
 	set +e
-        chmod +x Build/BatchFiles/Linux/GetAssets.py
-	Build/BatchFiles/Linux/GetAssets.py EpicGames/UnrealEngine $GITHUB_TAG 
+	# should not be needed since we add it as post-merge hook, but run here for now to overkill
+	mono Binaries/DotNET/GitDependencies.exe
 	UpdateResult=$?
 	set -e
 
diff --git a/Engine/Build/BatchFiles/Linux/GetAssets.py b/Engine/Build/BatchFiles/Linux/GetAssets.py
deleted file mode 100644
diff --git a/Engine/Build/BatchFiles/Linux/UpdateDeps.sh b/Engine/Build/BatchFiles/Linux/UpdateDeps.sh
--- a/Engine/Build/BatchFiles/Linux/UpdateDeps.sh
+++ b/Engine/Build/BatchFiles/Linux/UpdateDeps.sh
@@ -131,8 +131,8 @@ Unixify()
   cd ${TOP_DIR}
 }
 
-CleanRepo
-ExtractArchives
+#CleanRepo
+#ExtractArchives
 Unixify
 
 echo "********** SUCCESS ****************"
diff --git a/GitSetup.sh b/GitSetup.sh
--- a/GitSetup.sh
+++ b/GitSetup.sh
@@ -13,10 +13,14 @@ fi
 
 # Setup the git hooks
 if [ -d .git/hooks ]; then
-	echo Registering git hooks...
+	echo "Registering git hooks... (this will override existing ones!)"
 	echo \#!/bin/sh >.git/hooks/post-checkout
 	echo mono Engine/Binaries/DotNET/GitDependencies.exe >>.git/hooks/post-checkout
 	chmod +x .git/hooks/post-checkout
+
+	echo \#!/bin/sh >.git/hooks/post-merge
+	echo mono Engine/Binaries/DotNET/GitDependencies.exe >>.git/hooks/post-merge
+	chmod +x .git/hooks/post-merge
 fi
 
 #if [ "$(uname)" = "Darwin" ]; then
