Always fixup symlinks and executable permissions for whichever mono files are available

diff --git a/Engine/Build/BatchFiles/Mac/FixMonoFiles.sh b/Engine/Build/BatchFiles/Mac/FixMonoFiles.sh
--- a/Engine/Build/BatchFiles/Mac/FixMonoFiles.sh
+++ b/Engine/Build/BatchFiles/Mac/FixMonoFiles.sh
@@ -1,218 +1,232 @@
 #!/bin/sh
 
-if [ ! -f ../../../Binaries/ThirdParty/Mono/Mac/bin/mono ]; then
-  CUR_DIR=`pwd`
-  cd ../../../Binaries/ThirdParty/Mono/Mac
-  ln -s mono-boehm bin/mono
-  ln -s ../gac/Accessibility/4.0.0.0__b03f5f7f11d50a3a/Accessibility.dll lib/mono/4.5/Accessibility.dll
-  ln -s ../gac/Commons.Xml.Relaxng/4.0.0.0__0738eb9f132ed756/Commons.Xml.Relaxng.dll lib/mono/4.5/Commons.Xml.Relaxng.dll
-  ln -s ../gac/cscompmgd/0.0.0.0__b03f5f7f11d50a3a/cscompmgd.dll lib/mono/4.5/cscompmgd.dll
-  ln -s ../gac/CustomMarshalers/4.0.0.0__b03f5f7f11d50a3a/CustomMarshalers.dll lib/mono/4.5/CustomMarshalers.dll
-  ln -s ../gac/EntityFramework/6.0.0.0__b77a5c561934e089/EntityFramework.dll lib/mono/4.5/EntityFramework.dll
-  ln -s ../gac/EntityFramework.SqlServer/6.0.0.0__b77a5c561934e089/EntityFramework.SqlServer.dll lib/mono/4.5/EntityFramework.SqlServer.dll
-  ln -s ../gac/I18N.CJK/4.0.0.0__0738eb9f132ed756/I18N.CJK.dll lib/mono/4.5/I18N.CJK.dll
-  ln -s ../gac/I18N/4.0.0.0__0738eb9f132ed756/I18N.dll lib/mono/4.5/I18N.dll
-  ln -s ../gac/I18N.MidEast/4.0.0.0__0738eb9f132ed756/I18N.MidEast.dll lib/mono/4.5/I18N.MidEast.dll
-  ln -s ../gac/I18N.Other/4.0.0.0__0738eb9f132ed756/I18N.Other.dll lib/mono/4.5/I18N.Other.dll
-  ln -s ../gac/I18N.Rare/4.0.0.0__0738eb9f132ed756/I18N.Rare.dll lib/mono/4.5/I18N.Rare.dll
-  ln -s ../gac/I18N.West/4.0.0.0__0738eb9f132ed756/I18N.West.dll lib/mono/4.5/I18N.West.dll
-  ln -s ../gac/IBM.Data.DB2/1.0.0.0__7c307b91aa13d208/IBM.Data.DB2.dll lib/mono/4.5/IBM.Data.DB2.dll
-  ln -s ../gac/ICSharpCode.SharpZipLib/4.84.0.0__1b03e6acf1164f73/ICSharpCode.SharpZipLib.dll lib/mono/4.5/ICSharpCode.SharpZipLib.dll
-  ln -s ../gac/Microsoft.Build/4.0.0.0__b03f5f7f11d50a3a/Microsoft.Build.dll lib/mono/4.5/Microsoft.Build.dll
-  ln -s ../gac/Microsoft.Build.Engine/4.0.0.0__b03f5f7f11d50a3a/Microsoft.Build.Engine.dll lib/mono/4.5/Microsoft.Build.Engine.dll
-  ln -s ../gac/Microsoft.Build.Framework/4.0.0.0__b03f5f7f11d50a3a/Microsoft.Build.Framework.dll lib/mono/4.5/Microsoft.Build.Framework.dll
-  ln -s ../gac/Microsoft.Build.Tasks.v4.0/4.0.0.0__b03f5f7f11d50a3a/Microsoft.Build.Tasks.v4.0.dll lib/mono/4.5/Microsoft.Build.Tasks.v4.0.dll
-  ln -s ../gac/Microsoft.Build.Utilities.v4.0/4.0.0.0__b03f5f7f11d50a3a/Microsoft.Build.Utilities.v4.0.dll lib/mono/4.5/Microsoft.Build.Utilities.v4.0.dll
-  ln -s ../gac/Microsoft.CSharp/4.0.0.0__b03f5f7f11d50a3a/Microsoft.CSharp.dll lib/mono/4.5/Microsoft.CSharp.dll
-  ln -s ../gac/Microsoft.VisualC/0.0.0.0__b03f5f7f11d50a3a/Microsoft.VisualC.dll lib/mono/4.5/Microsoft.VisualC.dll
-  ln -s ../gac/Microsoft.Web.Infrastructure/1.0.0.0__31bf3856ad364e35/Microsoft.Web.Infrastructure.dll lib/mono/4.5/Microsoft.Web.Infrastructure.dll
-  ln -s ../gac/Mono.C5/1.1.1.0__ba07f434b1c35cbd/Mono.C5.dll lib/mono/4.5/Mono.C5.dll
-  ln -s ../gac/Mono.Cairo/4.0.0.0__0738eb9f132ed756/Mono.Cairo.dll lib/mono/4.5/Mono.Cairo.dll
-  ln -s ../gac/Mono.CodeContracts/4.0.0.0__0738eb9f132ed756/Mono.CodeContracts.dll lib/mono/4.5/Mono.CodeContracts.dll
-  ln -s ../gac/Mono.CompilerServices.SymbolWriter/4.0.0.0__0738eb9f132ed756/Mono.CompilerServices.SymbolWriter.dll lib/mono/4.5/Mono.CompilerServices.SymbolWriter.dll
-  ln -s ../gac/Mono.CSharp/4.0.0.0__0738eb9f132ed756/Mono.CSharp.dll lib/mono/4.5/Mono.CSharp.dll
-  ln -s ../gac/Mono.Data.Sqlite/4.0.0.0__0738eb9f132ed756/Mono.Data.Sqlite.dll lib/mono/4.5/Mono.Data.Sqlite.dll
-  ln -s ../gac/Mono.Data.Tds/4.0.0.0__0738eb9f132ed756/Mono.Data.Tds.dll lib/mono/4.5/Mono.Data.Tds.dll
-  ln -s ../gac/Mono.Debugger.Soft/4.0.0.0__0738eb9f132ed756/Mono.Debugger.Soft.dll lib/mono/4.5/Mono.Debugger.Soft.dll
-  ln -s ../gac/Mono.Http/4.0.0.0__0738eb9f132ed756/Mono.Http.dll lib/mono/4.5/Mono.Http.dll
-  ln -s ../gac/Mono.Management/4.0.0.0__0738eb9f132ed756/Mono.Management.dll lib/mono/4.5/Mono.Management.dll
-  ln -s ../gac/Mono.Messaging/4.0.0.0__0738eb9f132ed756/Mono.Messaging.dll lib/mono/4.5/Mono.Messaging.dll
-  ln -s ../gac/Mono.Messaging.RabbitMQ/4.0.0.0__0738eb9f132ed756/Mono.Messaging.RabbitMQ.dll lib/mono/4.5/Mono.Messaging.RabbitMQ.dll
-  ln -s ../gac/Mono.Parallel/4.0.0.0__0738eb9f132ed756/Mono.Parallel.dll lib/mono/4.5/Mono.Parallel.dll
-  ln -s ../gac/Mono.Posix/4.0.0.0__0738eb9f132ed756/Mono.Posix.dll lib/mono/4.5/Mono.Posix.dll
-  ln -s ../gac/Mono.Security/4.0.0.0__0738eb9f132ed756/Mono.Security.dll lib/mono/4.5/Mono.Security.dll
-  ln -s ../gac/Mono.Security.Win32/4.0.0.0__0738eb9f132ed756/Mono.Security.Win32.dll lib/mono/4.5/Mono.Security.Win32.dll
-  ln -s ../gac/Mono.Simd/4.0.0.0__0738eb9f132ed756/Mono.Simd.dll lib/mono/4.5/Mono.Simd.dll
-  ln -s ../gac/Mono.Tasklets/4.0.0.0__0738eb9f132ed756/Mono.Tasklets.dll lib/mono/4.5/Mono.Tasklets.dll
-  ln -s ../gac/Mono.Web/4.0.0.0__0738eb9f132ed756/Mono.Web.dll lib/mono/4.5/Mono.Web.dll
-  ln -s ../gac/Mono.WebBrowser/4.0.0.0__0738eb9f132ed756/Mono.WebBrowser.dll lib/mono/4.5/Mono.WebBrowser.dll
-  ln -s ../gac/Novell.Directory.Ldap/4.0.0.0__0738eb9f132ed756/Novell.Directory.Ldap.dll lib/mono/4.5/Novell.Directory.Ldap.dll
-  ln -s ../gac/Npgsql/4.0.0.0__5d8b90d52f46fda7/Npgsql.dll lib/mono/4.5/Npgsql.dll
-  ln -s ../gac/OpenSystem.C/4.0.0.0__b77a5c561934e089/OpenSystem.C.dll lib/mono/4.5/OpenSystem.C.dll
-  ln -s ../gac/PEAPI/4.0.0.0__0738eb9f132ed756/PEAPI.dll lib/mono/4.5/PEAPI.dll
-  ln -s ../gac/RabbitMQ.Client/4.0.0.0__b03f5f7f11d50a3a/RabbitMQ.Client.dll lib/mono/4.5/RabbitMQ.Client.dll
-  ln -s ../gac/System.ComponentModel.Composition/4.0.0.0__b77a5c561934e089/System.ComponentModel.Composition.dll lib/mono/4.5/System.ComponentModel.Composition.dll
-  ln -s ../gac/System.ComponentModel.DataAnnotations/4.0.0.0__31bf3856ad364e35/System.ComponentModel.DataAnnotations.dll lib/mono/4.5/System.ComponentModel.DataAnnotations.dll
-  ln -s ../gac/System.Configuration/4.0.0.0__b03f5f7f11d50a3a/System.Configuration.dll lib/mono/4.5/System.Configuration.dll
-  ln -s ../gac/System.Configuration.Install/4.0.0.0__b03f5f7f11d50a3a/System.Configuration.Install.dll lib/mono/4.5/System.Configuration.Install.dll
-  ln -s ../gac/System.Core/4.0.0.0__b77a5c561934e089/System.Core.dll lib/mono/4.5/System.Core.dll
-  ln -s ../gac/System.Data.DataSetExtensions/4.0.0.0__b77a5c561934e089/System.Data.DataSetExtensions.dll lib/mono/4.5/System.Data.DataSetExtensions.dll
-  ln -s ../gac/System.Data/4.0.0.0__b77a5c561934e089/System.Data.dll lib/mono/4.5/System.Data.dll
-  ln -s ../gac/System.Data.Linq/4.0.0.0__b77a5c561934e089/System.Data.Linq.dll lib/mono/4.5/System.Data.Linq.dll
-  ln -s ../gac/System.Data.OracleClient/4.0.0.0__b77a5c561934e089/System.Data.OracleClient.dll lib/mono/4.5/System.Data.OracleClient.dll
-  ln -s ../gac/System.Data.Services.Client/4.0.0.0__b77a5c561934e089/System.Data.Services.Client.dll lib/mono/4.5/System.Data.Services.Client.dll
-  ln -s ../gac/System.Data.Services/4.0.0.0__b77a5c561934e089/System.Data.Services.dll lib/mono/4.5/System.Data.Services.dll
-  ln -s ../gac/System.Design/4.0.0.0__b03f5f7f11d50a3a/System.Design.dll lib/mono/4.5/System.Design.dll
-  ln -s ../gac/System.DirectoryServices/4.0.0.0__b03f5f7f11d50a3a/System.DirectoryServices.dll lib/mono/4.5/System.DirectoryServices.dll
-  ln -s ../gac/System.DirectoryServices.Protocols/4.0.0.0__b03f5f7f11d50a3a/System.DirectoryServices.Protocols.dll lib/mono/4.5/System.DirectoryServices.Protocols.dll
-  ln -s ../gac/System/4.0.0.0__b77a5c561934e089/System.dll lib/mono/4.5/System.dll
-  ln -s ../gac/System.Drawing.Design/4.0.0.0__b03f5f7f11d50a3a/System.Drawing.Design.dll lib/mono/4.5/System.Drawing.Design.dll
-  ln -s ../gac/System.Drawing/4.0.0.0__b03f5f7f11d50a3a/System.Drawing.dll lib/mono/4.5/System.Drawing.dll
-  ln -s ../gac/System.Dynamic/4.0.0.0__b77a5c561934e089/System.Dynamic.dll lib/mono/4.5/System.Dynamic.dll
-  ln -s ../gac/System.EnterpriseServices/4.0.0.0__b03f5f7f11d50a3a/System.EnterpriseServices.dll lib/mono/4.5/System.EnterpriseServices.dll
-  ln -s ../gac/System.IdentityModel/4.0.0.0__b77a5c561934e089/System.IdentityModel.dll lib/mono/4.5/System.IdentityModel.dll
-  ln -s ../gac/System.IdentityModel.Selectors/4.0.0.0__b77a5c561934e089/System.IdentityModel.Selectors.dll lib/mono/4.5/System.IdentityModel.Selectors.dll
-  ln -s ../gac/System.IO.Compression/4.0.0.0__b77a5c561934e089/System.IO.Compression.dll lib/mono/4.5/System.IO.Compression.dll
-  ln -s ../gac/System.IO.Compression.FileSystem/4.0.0.0__b77a5c561934e089/System.IO.Compression.FileSystem.dll lib/mono/4.5/System.IO.Compression.FileSystem.dll
-  ln -s ../gac/System.Json/4.0.0.0__31bf3856ad364e35/System.Json.dll lib/mono/4.5/System.Json.dll
-  ln -s ../gac/System.Json.Microsoft/4.0.0.0__31bf3856ad364e35/System.Json.Microsoft.dll lib/mono/4.5/System.Json.Microsoft.dll
-  ln -s ../gac/System.Management/4.0.0.0__b03f5f7f11d50a3a/System.Management.dll lib/mono/4.5/System.Management.dll
-  ln -s ../gac/System.Messaging/4.0.0.0__b03f5f7f11d50a3a/System.Messaging.dll lib/mono/4.5/System.Messaging.dll
-  ln -s ../gac/System.Net/4.0.0.0__b03f5f7f11d50a3a/System.Net.dll lib/mono/4.5/System.Net.dll
-  ln -s ../gac/System.Net.Http/4.0.0.0__b03f5f7f11d50a3a/System.Net.Http.dll lib/mono/4.5/System.Net.Http.dll
-  ln -s ../gac/System.Net.Http.Formatting/4.0.0.0__31bf3856ad364e35/System.Net.Http.Formatting.dll lib/mono/4.5/System.Net.Http.Formatting.dll
-  ln -s ../gac/System.Net.Http.WebRequest/4.0.0.0__b03f5f7f11d50a3a/System.Net.Http.WebRequest.dll lib/mono/4.5/System.Net.Http.WebRequest.dll
-  ln -s ../gac/System.Numerics/4.0.0.0__b77a5c561934e089/System.Numerics.dll lib/mono/4.5/System.Numerics.dll
-  ln -s ../gac/System.Reactive.Core/2.1.30214.0__31bf3856ad364e35/System.Reactive.Core.dll lib/mono/4.5/System.Reactive.Core.dll
-  ln -s ../gac/System.Reactive.Debugger/2.1.30214.0__31bf3856ad364e35/System.Reactive.Debugger.dll lib/mono/4.5/System.Reactive.Debugger.dll
-  ln -s ../gac/System.Reactive.Experimental/2.1.30214.0__31bf3856ad364e35/System.Reactive.Experimental.dll lib/mono/4.5/System.Reactive.Experimental.dll
-  ln -s ../gac/System.Reactive.Interfaces/2.1.30214.0__31bf3856ad364e35/System.Reactive.Interfaces.dll lib/mono/4.5/System.Reactive.Interfaces.dll
-  ln -s ../gac/System.Reactive.Linq/2.1.30214.0__31bf3856ad364e35/System.Reactive.Linq.dll lib/mono/4.5/System.Reactive.Linq.dll
-  ln -s ../gac/System.Reactive.PlatformServices/2.1.30214.0__31bf3856ad364e35/System.Reactive.PlatformServices.dll lib/mono/4.5/System.Reactive.PlatformServices.dll
-  ln -s ../gac/System.Reactive.Providers/2.1.30214.0__31bf3856ad364e35/System.Reactive.Providers.dll lib/mono/4.5/System.Reactive.Providers.dll
-  ln -s ../gac/System.Reactive.Runtime.Remoting/2.1.30214.0__31bf3856ad364e35/System.Reactive.Runtime.Remoting.dll lib/mono/4.5/System.Reactive.Runtime.Remoting.dll
-  ln -s ../gac/System.Reactive.Windows.Forms/2.1.30214.0__31bf3856ad364e35/System.Reactive.Windows.Forms.dll lib/mono/4.5/System.Reactive.Windows.Forms.dll
-  ln -s ../gac/System.Reactive.Windows.Threading/2.1.30214.0__31bf3856ad364e35/System.Reactive.Windows.Threading.dll lib/mono/4.5/System.Reactive.Windows.Threading.dll
-  ln -s ../gac/System.Runtime.Caching/4.0.0.0__b03f5f7f11d50a3a/System.Runtime.Caching.dll lib/mono/4.5/System.Runtime.Caching.dll
-  ln -s ../gac/System.Runtime.DurableInstancing/4.0.0.0__31bf3856ad364e35/System.Runtime.DurableInstancing.dll lib/mono/4.5/System.Runtime.DurableInstancing.dll
-  ln -s ../gac/System.Runtime.Remoting/4.0.0.0__b77a5c561934e089/System.Runtime.Remoting.dll lib/mono/4.5/System.Runtime.Remoting.dll
-  ln -s ../gac/System.Runtime.Serialization/4.0.0.0__b77a5c561934e089/System.Runtime.Serialization.dll lib/mono/4.5/System.Runtime.Serialization.dll
-  ln -s ../gac/System.Runtime.Serialization.Formatters.Soap/4.0.0.0__b03f5f7f11d50a3a/System.Runtime.Serialization.Formatters.Soap.dll lib/mono/4.5/System.Runtime.Serialization.Formatters.Soap.dll
-  ln -s ../gac/System.Security/4.0.0.0__b03f5f7f11d50a3a/System.Security.dll lib/mono/4.5/System.Security.dll
-  ln -s ../gac/System.ServiceModel.Activation/4.0.0.0__31bf3856ad364e35/System.ServiceModel.Activation.dll lib/mono/4.5/System.ServiceModel.Activation.dll
-  ln -s ../gac/System.ServiceModel.Discovery/4.0.0.0__31bf3856ad364e35/System.ServiceModel.Discovery.dll lib/mono/4.5/System.ServiceModel.Discovery.dll
-  ln -s ../gac/System.ServiceModel/4.0.0.0__b77a5c561934e089/System.ServiceModel.dll lib/mono/4.5/System.ServiceModel.dll
-  ln -s ../gac/System.ServiceModel.Routing/4.0.0.0__31bf3856ad364e35/System.ServiceModel.Routing.dll lib/mono/4.5/System.ServiceModel.Routing.dll
-  ln -s ../gac/System.ServiceModel.Web/4.0.0.0__31bf3856ad364e35/System.ServiceModel.Web.dll lib/mono/4.5/System.ServiceModel.Web.dll
-  ln -s ../gac/System.ServiceProcess/4.0.0.0__b03f5f7f11d50a3a/System.ServiceProcess.dll lib/mono/4.5/System.ServiceProcess.dll
-  ln -s ../gac/System.Threading.Tasks.Dataflow/4.0.0.0__b77a5c561934e089/System.Threading.Tasks.Dataflow.dll lib/mono/4.5/System.Threading.Tasks.Dataflow.dll
-  ln -s ../gac/System.Transactions/4.0.0.0__b77a5c561934e089/System.Transactions.dll lib/mono/4.5/System.Transactions.dll
-  ln -s ../gac/System.Web.Abstractions/4.0.0.0__31bf3856ad364e35/System.Web.Abstractions.dll lib/mono/4.5/System.Web.Abstractions.dll
-  ln -s ../gac/System.Web.ApplicationServices/4.0.0.0__31bf3856ad364e35/System.Web.ApplicationServices.dll lib/mono/4.5/System.Web.ApplicationServices.dll
-  ln -s ../gac/System.Web/4.0.0.0__b03f5f7f11d50a3a/System.Web.dll lib/mono/4.5/System.Web.dll
-  ln -s ../gac/System.Web.DynamicData/4.0.0.0__31bf3856ad364e35/System.Web.DynamicData.dll lib/mono/4.5/System.Web.DynamicData.dll
-  ln -s ../gac/System.Web.Extensions.Design/4.0.0.0__31bf3856ad364e35/System.Web.Extensions.Design.dll lib/mono/4.5/System.Web.Extensions.Design.dll
-  ln -s ../gac/System.Web.Extensions/4.0.0.0__31bf3856ad364e35/System.Web.Extensions.dll lib/mono/4.5/System.Web.Extensions.dll
-  ln -s ../gac/System.Web.Http/4.0.0.0__31bf3856ad364e35/System.Web.Http.dll lib/mono/4.5/System.Web.Http.dll
-  ln -s ../gac/System.Web.Http.SelfHost/4.0.0.0__31bf3856ad364e35/System.Web.Http.SelfHost.dll lib/mono/4.5/System.Web.Http.SelfHost.dll
-  ln -s ../gac/System.Web.Http.WebHost/4.0.0.0__31bf3856ad364e35/System.Web.Http.WebHost.dll lib/mono/4.5/System.Web.Http.WebHost.dll
-  ln -s ../gac/System.Web.Mvc/3.0.0.0__31bf3856ad364e35/System.Web.Mvc.dll lib/mono/4.5/System.Web.Mvc.dll
-  ln -s ../gac/System.Web.Razor/2.0.0.0__31bf3856ad364e35/System.Web.Razor.dll lib/mono/4.5/System.Web.Razor.dll
-  ln -s ../gac/System.Web.Routing/4.0.0.0__31bf3856ad364e35/System.Web.Routing.dll lib/mono/4.5/System.Web.Routing.dll
-  ln -s ../gac/System.Web.Services/4.0.0.0__b03f5f7f11d50a3a/System.Web.Services.dll lib/mono/4.5/System.Web.Services.dll
-  ln -s ../gac/System.Web.WebPages.Deployment/2.0.0.0__31bf3856ad364e35/System.Web.WebPages.Deployment.dll lib/mono/4.5/System.Web.WebPages.Deployment.dll
-  ln -s ../gac/System.Web.WebPages/2.0.0.0__31bf3856ad364e35/System.Web.WebPages.dll lib/mono/4.5/System.Web.WebPages.dll
-  ln -s ../gac/System.Web.WebPages.Razor/2.0.0.0__31bf3856ad364e35/System.Web.WebPages.Razor.dll lib/mono/4.5/System.Web.WebPages.Razor.dll
-  ln -s ../gac/System.Windows/4.0.0.0__b03f5f7f11d50a3a/System.Windows.dll lib/mono/4.5/System.Windows.dll
-  ln -s ../gac/System.Windows.Forms.DataVisualization/4.0.0.0__b77a5c561934e089/System.Windows.Forms.DataVisualization.dll lib/mono/4.5/System.Windows.Forms.DataVisualization.dll
-  ln -s ../gac/System.Windows.Forms/4.0.0.0__b77a5c561934e089/System.Windows.Forms.dll lib/mono/4.5/System.Windows.Forms.dll
-  ln -s ../gac/System.Xaml/4.0.0.0__b77a5c561934e089/System.Xaml.dll lib/mono/4.5/System.Xaml.dll
-  ln -s ../gac/System.Xml/4.0.0.0__b77a5c561934e089/System.Xml.dll lib/mono/4.5/System.Xml.dll
-  ln -s ../gac/System.Xml.Linq/4.0.0.0__b77a5c561934e089/System.Xml.Linq.dll lib/mono/4.5/System.Xml.Linq.dll
-  ln -s ../gac/System.Xml.Serialization/4.0.0.0__b77a5c561934e089/System.Xml.Serialization.dll lib/mono/4.5/System.Xml.Serialization.dll
-  ln -s ../gac/WebMatrix.Data/4.0.0.0__0738eb9f132ed756/WebMatrix.Data.dll lib/mono/4.5/WebMatrix.Data.dll
-  ln -s ../gac/WindowsBase/4.0.0.0__31bf3856ad364e35/WindowsBase.dll lib/mono/4.5/WindowsBase.dll
-  if [ ! -d lib/mono/monodoc ]; then
-    mkdir lib/mono/monodoc
-  fi
-  ln -s ../gac/monodoc/1.0.0.0__0738eb9f132ed756/monodoc.dll lib/mono/monodoc/monodoc.dll
-  chmod +x bin/al
-  chmod +x bin/al2
-  chmod +x bin/caspol
-  chmod +x bin/cccheck
-  chmod +x bin/ccrewrite
-  chmod +x bin/cert2spc
-  chmod +x bin/certmgr
-  chmod +x bin/chktrust
-  chmod +x bin/crlupdate
-  chmod +x bin/csharp
-  chmod +x bin/disco
-  chmod +x bin/dmcs
-  chmod +x bin/dtd2rng
-  chmod +x bin/dtd2xsd
-  chmod +x bin/gacutil
-  chmod +x bin/gacutil2
-  chmod +x bin/genxs
-  chmod +x bin/gmcs
-  chmod +x bin/httpcfg
-  chmod +x bin/ilasm
-  chmod +x bin/installvst
-  chmod +x bin/lc
-  chmod +x bin/macpack
-  chmod +x bin/makecert
-  chmod +x bin/mconfig
-  chmod +x bin/mcs
-  chmod +x bin/mdassembler
-  chmod +x bin/mdbrebase
-  chmod +x bin/mdoc
-  chmod +x bin/mdoc-assemble
-  chmod +x bin/mdoc-export-html
-  chmod +x bin/mdoc-export-msxdoc
-  chmod +x bin/mdoc-update
-  chmod +x bin/mdoc-validate
-  chmod +x bin/mdvalidater
-  chmod +x bin/mkbundle
-  chmod +x bin/mod
-  chmod +x bin/mono-api-info
-  chmod +x bin/mono-boehm
-  chmod +x bin/mono-cil-strip
-  chmod +x bin/mono-configuration-crypto
-  chmod +x bin/mono-find-provides
-  chmod +x bin/mono-find-requires
-  chmod +x bin/mono-heapviz
-  chmod +x bin/mono-service
-  chmod +x bin/mono-service2
-  chmod +x bin/mono-shlib-cop
-  chmod +x bin/mono-test-install
-  chmod +x bin/mono-xmltool
-  chmod +x bin/monodocer
-  chmod +x bin/monodocs2html
-  chmod +x bin/monodocs2slashdoc
-  chmod +x bin/monolinker
-  chmod +x bin/monop
-  chmod +x bin/monop2
-  chmod +x bin/mozroots
-  chmod +x bin/nunit-console
-  chmod +x bin/nunit-console2
-  chmod +x bin/nunit-console4
-  chmod +x bin/pdb2mdb
-  chmod +x bin/permview
-  chmod +x bin/peverify
-  chmod +x bin/prj2make
-  chmod +x bin/resgen
-  chmod +x bin/resgen2
-  chmod +x bin/secutil
-  chmod +x bin/setreg
-  chmod +x bin/sgen
-  chmod +x bin/signcode
-  chmod +x bin/sn
-  chmod +x bin/soapsuds
-  chmod +x bin/sqlmetal
-  chmod +x bin/sqlsharp
-  chmod +x bin/svcutil
-  chmod +x bin/wsdl
-  chmod +x bin/wsdl2
-  chmod +x bin/xbuild
-  chmod +x bin/xsd
-  cd "$CUR_DIR"
+function try_symlink {
+	if [ -f $(dirname $2)/$1 ]; then
+		if [ ! -f $2 ]; then
+			ln -s $1 $2
+		fi
+	fi
+}
+
+CUR_DIR=`pwd`
+cd ../../../Binaries/ThirdParty/Mono/Mac
+
+try_symlink ../gac/Accessibility/4.0.0.0__b03f5f7f11d50a3a/Accessibility.dll lib/mono/4.5/Accessibility.dll
+try_symlink ../gac/Commons.Xml.Relaxng/4.0.0.0__0738eb9f132ed756/Commons.Xml.Relaxng.dll lib/mono/4.5/Commons.Xml.Relaxng.dll
+try_symlink ../gac/cscompmgd/0.0.0.0__b03f5f7f11d50a3a/cscompmgd.dll lib/mono/4.5/cscompmgd.dll
+try_symlink ../gac/CustomMarshalers/4.0.0.0__b03f5f7f11d50a3a/CustomMarshalers.dll lib/mono/4.5/CustomMarshalers.dll
+try_symlink ../gac/EntityFramework/6.0.0.0__b77a5c561934e089/EntityFramework.dll lib/mono/4.5/EntityFramework.dll
+try_symlink ../gac/EntityFramework.SqlServer/6.0.0.0__b77a5c561934e089/EntityFramework.SqlServer.dll lib/mono/4.5/EntityFramework.SqlServer.dll
+try_symlink ../gac/I18N.CJK/4.0.0.0__0738eb9f132ed756/I18N.CJK.dll lib/mono/4.5/I18N.CJK.dll
+try_symlink ../gac/I18N/4.0.0.0__0738eb9f132ed756/I18N.dll lib/mono/4.5/I18N.dll
+try_symlink ../gac/I18N.MidEast/4.0.0.0__0738eb9f132ed756/I18N.MidEast.dll lib/mono/4.5/I18N.MidEast.dll
+try_symlink ../gac/I18N.Other/4.0.0.0__0738eb9f132ed756/I18N.Other.dll lib/mono/4.5/I18N.Other.dll
+try_symlink ../gac/I18N.Rare/4.0.0.0__0738eb9f132ed756/I18N.Rare.dll lib/mono/4.5/I18N.Rare.dll
+try_symlink ../gac/I18N.West/4.0.0.0__0738eb9f132ed756/I18N.West.dll lib/mono/4.5/I18N.West.dll
+try_symlink ../gac/IBM.Data.DB2/1.0.0.0__7c307b91aa13d208/IBM.Data.DB2.dll lib/mono/4.5/IBM.Data.DB2.dll
+try_symlink ../gac/ICSharpCode.SharpZipLib/4.84.0.0__1b03e6acf1164f73/ICSharpCode.SharpZipLib.dll lib/mono/4.5/ICSharpCode.SharpZipLib.dll
+try_symlink ../gac/Microsoft.Build/4.0.0.0__b03f5f7f11d50a3a/Microsoft.Build.dll lib/mono/4.5/Microsoft.Build.dll
+try_symlink ../gac/Microsoft.Build.Engine/4.0.0.0__b03f5f7f11d50a3a/Microsoft.Build.Engine.dll lib/mono/4.5/Microsoft.Build.Engine.dll
+try_symlink ../gac/Microsoft.Build.Framework/4.0.0.0__b03f5f7f11d50a3a/Microsoft.Build.Framework.dll lib/mono/4.5/Microsoft.Build.Framework.dll
+try_symlink ../gac/Microsoft.Build.Tasks.v4.0/4.0.0.0__b03f5f7f11d50a3a/Microsoft.Build.Tasks.v4.0.dll lib/mono/4.5/Microsoft.Build.Tasks.v4.0.dll
+try_symlink ../gac/Microsoft.Build.Utilities.v4.0/4.0.0.0__b03f5f7f11d50a3a/Microsoft.Build.Utilities.v4.0.dll lib/mono/4.5/Microsoft.Build.Utilities.v4.0.dll
+try_symlink ../gac/Microsoft.CSharp/4.0.0.0__b03f5f7f11d50a3a/Microsoft.CSharp.dll lib/mono/4.5/Microsoft.CSharp.dll
+try_symlink ../gac/Microsoft.VisualC/0.0.0.0__b03f5f7f11d50a3a/Microsoft.VisualC.dll lib/mono/4.5/Microsoft.VisualC.dll
+try_symlink ../gac/Microsoft.Web.Infrastructure/1.0.0.0__31bf3856ad364e35/Microsoft.Web.Infrastructure.dll lib/mono/4.5/Microsoft.Web.Infrastructure.dll
+try_symlink ../gac/Mono.C5/1.1.1.0__ba07f434b1c35cbd/Mono.C5.dll lib/mono/4.5/Mono.C5.dll
+try_symlink ../gac/Mono.Cairo/4.0.0.0__0738eb9f132ed756/Mono.Cairo.dll lib/mono/4.5/Mono.Cairo.dll
+try_symlink ../gac/Mono.CodeContracts/4.0.0.0__0738eb9f132ed756/Mono.CodeContracts.dll lib/mono/4.5/Mono.CodeContracts.dll
+try_symlink ../gac/Mono.CompilerServices.SymbolWriter/4.0.0.0__0738eb9f132ed756/Mono.CompilerServices.SymbolWriter.dll lib/mono/4.5/Mono.CompilerServices.SymbolWriter.dll
+try_symlink ../gac/Mono.CSharp/4.0.0.0__0738eb9f132ed756/Mono.CSharp.dll lib/mono/4.5/Mono.CSharp.dll
+try_symlink ../gac/Mono.Data.Sqlite/4.0.0.0__0738eb9f132ed756/Mono.Data.Sqlite.dll lib/mono/4.5/Mono.Data.Sqlite.dll
+try_symlink ../gac/Mono.Data.Tds/4.0.0.0__0738eb9f132ed756/Mono.Data.Tds.dll lib/mono/4.5/Mono.Data.Tds.dll
+try_symlink ../gac/Mono.Debugger.Soft/4.0.0.0__0738eb9f132ed756/Mono.Debugger.Soft.dll lib/mono/4.5/Mono.Debugger.Soft.dll
+try_symlink ../gac/Mono.Http/4.0.0.0__0738eb9f132ed756/Mono.Http.dll lib/mono/4.5/Mono.Http.dll
+try_symlink ../gac/Mono.Management/4.0.0.0__0738eb9f132ed756/Mono.Management.dll lib/mono/4.5/Mono.Management.dll
+try_symlink ../gac/Mono.Messaging/4.0.0.0__0738eb9f132ed756/Mono.Messaging.dll lib/mono/4.5/Mono.Messaging.dll
+try_symlink ../gac/Mono.Messaging.RabbitMQ/4.0.0.0__0738eb9f132ed756/Mono.Messaging.RabbitMQ.dll lib/mono/4.5/Mono.Messaging.RabbitMQ.dll
+try_symlink ../gac/Mono.Parallel/4.0.0.0__0738eb9f132ed756/Mono.Parallel.dll lib/mono/4.5/Mono.Parallel.dll
+try_symlink ../gac/Mono.Posix/4.0.0.0__0738eb9f132ed756/Mono.Posix.dll lib/mono/4.5/Mono.Posix.dll
+try_symlink ../gac/Mono.Security/4.0.0.0__0738eb9f132ed756/Mono.Security.dll lib/mono/4.5/Mono.Security.dll
+try_symlink ../gac/Mono.Security.Win32/4.0.0.0__0738eb9f132ed756/Mono.Security.Win32.dll lib/mono/4.5/Mono.Security.Win32.dll
+try_symlink ../gac/Mono.Simd/4.0.0.0__0738eb9f132ed756/Mono.Simd.dll lib/mono/4.5/Mono.Simd.dll
+try_symlink ../gac/Mono.Tasklets/4.0.0.0__0738eb9f132ed756/Mono.Tasklets.dll lib/mono/4.5/Mono.Tasklets.dll
+try_symlink ../gac/Mono.Web/4.0.0.0__0738eb9f132ed756/Mono.Web.dll lib/mono/4.5/Mono.Web.dll
+try_symlink ../gac/Mono.WebBrowser/4.0.0.0__0738eb9f132ed756/Mono.WebBrowser.dll lib/mono/4.5/Mono.WebBrowser.dll
+try_symlink ../gac/Novell.Directory.Ldap/4.0.0.0__0738eb9f132ed756/Novell.Directory.Ldap.dll lib/mono/4.5/Novell.Directory.Ldap.dll
+try_symlink ../gac/Npgsql/4.0.0.0__5d8b90d52f46fda7/Npgsql.dll lib/mono/4.5/Npgsql.dll
+try_symlink ../gac/OpenSystem.C/4.0.0.0__b77a5c561934e089/OpenSystem.C.dll lib/mono/4.5/OpenSystem.C.dll
+try_symlink ../gac/PEAPI/4.0.0.0__0738eb9f132ed756/PEAPI.dll lib/mono/4.5/PEAPI.dll
+try_symlink ../gac/RabbitMQ.Client/4.0.0.0__b03f5f7f11d50a3a/RabbitMQ.Client.dll lib/mono/4.5/RabbitMQ.Client.dll
+try_symlink ../gac/System.ComponentModel.Composition/4.0.0.0__b77a5c561934e089/System.ComponentModel.Composition.dll lib/mono/4.5/System.ComponentModel.Composition.dll
+try_symlink ../gac/System.ComponentModel.DataAnnotations/4.0.0.0__31bf3856ad364e35/System.ComponentModel.DataAnnotations.dll lib/mono/4.5/System.ComponentModel.DataAnnotations.dll
+try_symlink ../gac/System.Configuration/4.0.0.0__b03f5f7f11d50a3a/System.Configuration.dll lib/mono/4.5/System.Configuration.dll
+try_symlink ../gac/System.Configuration.Install/4.0.0.0__b03f5f7f11d50a3a/System.Configuration.Install.dll lib/mono/4.5/System.Configuration.Install.dll
+try_symlink ../gac/System.Core/4.0.0.0__b77a5c561934e089/System.Core.dll lib/mono/4.5/System.Core.dll
+try_symlink ../gac/System.Data.DataSetExtensions/4.0.0.0__b77a5c561934e089/System.Data.DataSetExtensions.dll lib/mono/4.5/System.Data.DataSetExtensions.dll
+try_symlink ../gac/System.Data/4.0.0.0__b77a5c561934e089/System.Data.dll lib/mono/4.5/System.Data.dll
+try_symlink ../gac/System.Data.Linq/4.0.0.0__b77a5c561934e089/System.Data.Linq.dll lib/mono/4.5/System.Data.Linq.dll
+try_symlink ../gac/System.Data.OracleClient/4.0.0.0__b77a5c561934e089/System.Data.OracleClient.dll lib/mono/4.5/System.Data.OracleClient.dll
+try_symlink ../gac/System.Data.Services.Client/4.0.0.0__b77a5c561934e089/System.Data.Services.Client.dll lib/mono/4.5/System.Data.Services.Client.dll
+try_symlink ../gac/System.Data.Services/4.0.0.0__b77a5c561934e089/System.Data.Services.dll lib/mono/4.5/System.Data.Services.dll
+try_symlink ../gac/System.Design/4.0.0.0__b03f5f7f11d50a3a/System.Design.dll lib/mono/4.5/System.Design.dll
+try_symlink ../gac/System.DirectoryServices/4.0.0.0__b03f5f7f11d50a3a/System.DirectoryServices.dll lib/mono/4.5/System.DirectoryServices.dll
+try_symlink ../gac/System.DirectoryServices.Protocols/4.0.0.0__b03f5f7f11d50a3a/System.DirectoryServices.Protocols.dll lib/mono/4.5/System.DirectoryServices.Protocols.dll
+try_symlink ../gac/System/4.0.0.0__b77a5c561934e089/System.dll lib/mono/4.5/System.dll
+try_symlink ../gac/System.Drawing.Design/4.0.0.0__b03f5f7f11d50a3a/System.Drawing.Design.dll lib/mono/4.5/System.Drawing.Design.dll
+try_symlink ../gac/System.Drawing/4.0.0.0__b03f5f7f11d50a3a/System.Drawing.dll lib/mono/4.5/System.Drawing.dll
+try_symlink ../gac/System.Dynamic/4.0.0.0__b77a5c561934e089/System.Dynamic.dll lib/mono/4.5/System.Dynamic.dll
+try_symlink ../gac/System.EnterpriseServices/4.0.0.0__b03f5f7f11d50a3a/System.EnterpriseServices.dll lib/mono/4.5/System.EnterpriseServices.dll
+try_symlink ../gac/System.IdentityModel/4.0.0.0__b77a5c561934e089/System.IdentityModel.dll lib/mono/4.5/System.IdentityModel.dll
+try_symlink ../gac/System.IdentityModel.Selectors/4.0.0.0__b77a5c561934e089/System.IdentityModel.Selectors.dll lib/mono/4.5/System.IdentityModel.Selectors.dll
+try_symlink ../gac/System.IO.Compression/4.0.0.0__b77a5c561934e089/System.IO.Compression.dll lib/mono/4.5/System.IO.Compression.dll
+try_symlink ../gac/System.IO.Compression.FileSystem/4.0.0.0__b77a5c561934e089/System.IO.Compression.FileSystem.dll lib/mono/4.5/System.IO.Compression.FileSystem.dll
+try_symlink ../gac/System.Json/4.0.0.0__31bf3856ad364e35/System.Json.dll lib/mono/4.5/System.Json.dll
+try_symlink ../gac/System.Json.Microsoft/4.0.0.0__31bf3856ad364e35/System.Json.Microsoft.dll lib/mono/4.5/System.Json.Microsoft.dll
+try_symlink ../gac/System.Management/4.0.0.0__b03f5f7f11d50a3a/System.Management.dll lib/mono/4.5/System.Management.dll
+try_symlink ../gac/System.Messaging/4.0.0.0__b03f5f7f11d50a3a/System.Messaging.dll lib/mono/4.5/System.Messaging.dll
+try_symlink ../gac/System.Net/4.0.0.0__b03f5f7f11d50a3a/System.Net.dll lib/mono/4.5/System.Net.dll
+try_symlink ../gac/System.Net.Http/4.0.0.0__b03f5f7f11d50a3a/System.Net.Http.dll lib/mono/4.5/System.Net.Http.dll
+try_symlink ../gac/System.Net.Http.Formatting/4.0.0.0__31bf3856ad364e35/System.Net.Http.Formatting.dll lib/mono/4.5/System.Net.Http.Formatting.dll
+try_symlink ../gac/System.Net.Http.WebRequest/4.0.0.0__b03f5f7f11d50a3a/System.Net.Http.WebRequest.dll lib/mono/4.5/System.Net.Http.WebRequest.dll
+try_symlink ../gac/System.Numerics/4.0.0.0__b77a5c561934e089/System.Numerics.dll lib/mono/4.5/System.Numerics.dll
+try_symlink ../gac/System.Reactive.Core/2.1.30214.0__31bf3856ad364e35/System.Reactive.Core.dll lib/mono/4.5/System.Reactive.Core.dll
+try_symlink ../gac/System.Reactive.Debugger/2.1.30214.0__31bf3856ad364e35/System.Reactive.Debugger.dll lib/mono/4.5/System.Reactive.Debugger.dll
+try_symlink ../gac/System.Reactive.Experimental/2.1.30214.0__31bf3856ad364e35/System.Reactive.Experimental.dll lib/mono/4.5/System.Reactive.Experimental.dll
+try_symlink ../gac/System.Reactive.Interfaces/2.1.30214.0__31bf3856ad364e35/System.Reactive.Interfaces.dll lib/mono/4.5/System.Reactive.Interfaces.dll
+try_symlink ../gac/System.Reactive.Linq/2.1.30214.0__31bf3856ad364e35/System.Reactive.Linq.dll lib/mono/4.5/System.Reactive.Linq.dll
+try_symlink ../gac/System.Reactive.PlatformServices/2.1.30214.0__31bf3856ad364e35/System.Reactive.PlatformServices.dll lib/mono/4.5/System.Reactive.PlatformServices.dll
+try_symlink ../gac/System.Reactive.Providers/2.1.30214.0__31bf3856ad364e35/System.Reactive.Providers.dll lib/mono/4.5/System.Reactive.Providers.dll
+try_symlink ../gac/System.Reactive.Runtime.Remoting/2.1.30214.0__31bf3856ad364e35/System.Reactive.Runtime.Remoting.dll lib/mono/4.5/System.Reactive.Runtime.Remoting.dll
+try_symlink ../gac/System.Reactive.Windows.Forms/2.1.30214.0__31bf3856ad364e35/System.Reactive.Windows.Forms.dll lib/mono/4.5/System.Reactive.Windows.Forms.dll
+try_symlink ../gac/System.Reactive.Windows.Threading/2.1.30214.0__31bf3856ad364e35/System.Reactive.Windows.Threading.dll lib/mono/4.5/System.Reactive.Windows.Threading.dll
+try_symlink ../gac/System.Runtime.Caching/4.0.0.0__b03f5f7f11d50a3a/System.Runtime.Caching.dll lib/mono/4.5/System.Runtime.Caching.dll
+try_symlink ../gac/System.Runtime.DurableInstancing/4.0.0.0__31bf3856ad364e35/System.Runtime.DurableInstancing.dll lib/mono/4.5/System.Runtime.DurableInstancing.dll
+try_symlink ../gac/System.Runtime.Remoting/4.0.0.0__b77a5c561934e089/System.Runtime.Remoting.dll lib/mono/4.5/System.Runtime.Remoting.dll
+try_symlink ../gac/System.Runtime.Serialization/4.0.0.0__b77a5c561934e089/System.Runtime.Serialization.dll lib/mono/4.5/System.Runtime.Serialization.dll
+try_symlink ../gac/System.Runtime.Serialization.Formatters.Soap/4.0.0.0__b03f5f7f11d50a3a/System.Runtime.Serialization.Formatters.Soap.dll lib/mono/4.5/System.Runtime.Serialization.Formatters.Soap.dll
+try_symlink ../gac/System.Security/4.0.0.0__b03f5f7f11d50a3a/System.Security.dll lib/mono/4.5/System.Security.dll
+try_symlink ../gac/System.ServiceModel.Activation/4.0.0.0__31bf3856ad364e35/System.ServiceModel.Activation.dll lib/mono/4.5/System.ServiceModel.Activation.dll
+try_symlink ../gac/System.ServiceModel.Discovery/4.0.0.0__31bf3856ad364e35/System.ServiceModel.Discovery.dll lib/mono/4.5/System.ServiceModel.Discovery.dll
+try_symlink ../gac/System.ServiceModel/4.0.0.0__b77a5c561934e089/System.ServiceModel.dll lib/mono/4.5/System.ServiceModel.dll
+try_symlink ../gac/System.ServiceModel.Routing/4.0.0.0__31bf3856ad364e35/System.ServiceModel.Routing.dll lib/mono/4.5/System.ServiceModel.Routing.dll
+try_symlink ../gac/System.ServiceModel.Web/4.0.0.0__31bf3856ad364e35/System.ServiceModel.Web.dll lib/mono/4.5/System.ServiceModel.Web.dll
+try_symlink ../gac/System.ServiceProcess/4.0.0.0__b03f5f7f11d50a3a/System.ServiceProcess.dll lib/mono/4.5/System.ServiceProcess.dll
+try_symlink ../gac/System.Threading.Tasks.Dataflow/4.0.0.0__b77a5c561934e089/System.Threading.Tasks.Dataflow.dll lib/mono/4.5/System.Threading.Tasks.Dataflow.dll
+try_symlink ../gac/System.Transactions/4.0.0.0__b77a5c561934e089/System.Transactions.dll lib/mono/4.5/System.Transactions.dll
+try_symlink ../gac/System.Web.Abstractions/4.0.0.0__31bf3856ad364e35/System.Web.Abstractions.dll lib/mono/4.5/System.Web.Abstractions.dll
+try_symlink ../gac/System.Web.ApplicationServices/4.0.0.0__31bf3856ad364e35/System.Web.ApplicationServices.dll lib/mono/4.5/System.Web.ApplicationServices.dll
+try_symlink ../gac/System.Web/4.0.0.0__b03f5f7f11d50a3a/System.Web.dll lib/mono/4.5/System.Web.dll
+try_symlink ../gac/System.Web.DynamicData/4.0.0.0__31bf3856ad364e35/System.Web.DynamicData.dll lib/mono/4.5/System.Web.DynamicData.dll
+try_symlink ../gac/System.Web.Extensions.Design/4.0.0.0__31bf3856ad364e35/System.Web.Extensions.Design.dll lib/mono/4.5/System.Web.Extensions.Design.dll
+try_symlink ../gac/System.Web.Extensions/4.0.0.0__31bf3856ad364e35/System.Web.Extensions.dll lib/mono/4.5/System.Web.Extensions.dll
+try_symlink ../gac/System.Web.Http/4.0.0.0__31bf3856ad364e35/System.Web.Http.dll lib/mono/4.5/System.Web.Http.dll
+try_symlink ../gac/System.Web.Http.SelfHost/4.0.0.0__31bf3856ad364e35/System.Web.Http.SelfHost.dll lib/mono/4.5/System.Web.Http.SelfHost.dll
+try_symlink ../gac/System.Web.Http.WebHost/4.0.0.0__31bf3856ad364e35/System.Web.Http.WebHost.dll lib/mono/4.5/System.Web.Http.WebHost.dll
+try_symlink ../gac/System.Web.Mvc/3.0.0.0__31bf3856ad364e35/System.Web.Mvc.dll lib/mono/4.5/System.Web.Mvc.dll
+try_symlink ../gac/System.Web.Razor/2.0.0.0__31bf3856ad364e35/System.Web.Razor.dll lib/mono/4.5/System.Web.Razor.dll
+try_symlink ../gac/System.Web.Routing/4.0.0.0__31bf3856ad364e35/System.Web.Routing.dll lib/mono/4.5/System.Web.Routing.dll
+try_symlink ../gac/System.Web.Services/4.0.0.0__b03f5f7f11d50a3a/System.Web.Services.dll lib/mono/4.5/System.Web.Services.dll
+try_symlink ../gac/System.Web.WebPages.Deployment/2.0.0.0__31bf3856ad364e35/System.Web.WebPages.Deployment.dll lib/mono/4.5/System.Web.WebPages.Deployment.dll
+try_symlink ../gac/System.Web.WebPages/2.0.0.0__31bf3856ad364e35/System.Web.WebPages.dll lib/mono/4.5/System.Web.WebPages.dll
+try_symlink ../gac/System.Web.WebPages.Razor/2.0.0.0__31bf3856ad364e35/System.Web.WebPages.Razor.dll lib/mono/4.5/System.Web.WebPages.Razor.dll
+try_symlink ../gac/System.Windows/4.0.0.0__b03f5f7f11d50a3a/System.Windows.dll lib/mono/4.5/System.Windows.dll
+try_symlink ../gac/System.Windows.Forms.DataVisualization/4.0.0.0__b77a5c561934e089/System.Windows.Forms.DataVisualization.dll lib/mono/4.5/System.Windows.Forms.DataVisualization.dll
+try_symlink ../gac/System.Windows.Forms/4.0.0.0__b77a5c561934e089/System.Windows.Forms.dll lib/mono/4.5/System.Windows.Forms.dll
+try_symlink ../gac/System.Xaml/4.0.0.0__b77a5c561934e089/System.Xaml.dll lib/mono/4.5/System.Xaml.dll
+try_symlink ../gac/System.Xml/4.0.0.0__b77a5c561934e089/System.Xml.dll lib/mono/4.5/System.Xml.dll
+try_symlink ../gac/System.Xml.Linq/4.0.0.0__b77a5c561934e089/System.Xml.Linq.dll lib/mono/4.5/System.Xml.Linq.dll
+try_symlink ../gac/System.Xml.Serialization/4.0.0.0__b77a5c561934e089/System.Xml.Serialization.dll lib/mono/4.5/System.Xml.Serialization.dll
+try_symlink ../gac/WebMatrix.Data/4.0.0.0__0738eb9f132ed756/WebMatrix.Data.dll lib/mono/4.5/WebMatrix.Data.dll
+try_symlink ../gac/WindowsBase/4.0.0.0__31bf3856ad364e35/WindowsBase.dll lib/mono/4.5/WindowsBase.dll
+if [ ! -d lib/mono/monodoc ]; then
+  mkdir lib/mono/monodoc
 fi
+try_symlink ../gac/monodoc/1.0.0.0__0738eb9f132ed756/monodoc.dll lib/mono/monodoc/monodoc.dll
+
+function try_chmod {
+	if [ -f $1 ]; then
+		chmod +x $1
+	fi
+}
+
+try_chmod bin/al
+try_chmod bin/al2
+try_chmod bin/caspol
+try_chmod bin/cccheck
+try_chmod bin/ccrewrite
+try_chmod bin/cert2spc
+try_chmod bin/certmgr
+try_chmod bin/chktrust
+try_chmod bin/crlupdate
+try_chmod bin/csharp
+try_chmod bin/disco
+try_chmod bin/dmcs
+try_chmod bin/dtd2rng
+try_chmod bin/dtd2xsd
+try_chmod bin/gacutil
+try_chmod bin/gacutil2
+try_chmod bin/genxs
+try_chmod bin/gmcs
+try_chmod bin/httpcfg
+try_chmod bin/ilasm
+try_chmod bin/installvst
+try_chmod bin/lc
+try_chmod bin/macpack
+try_chmod bin/makecert
+try_chmod bin/mconfig
+try_chmod bin/mcs
+try_chmod bin/mdassembler
+try_chmod bin/mdbrebase
+try_chmod bin/mdoc
+try_chmod bin/mdoc-assemble
+try_chmod bin/mdoc-export-html
+try_chmod bin/mdoc-export-msxdoc
+try_chmod bin/mdoc-update
+try_chmod bin/mdoc-validate
+try_chmod bin/mdvalidater
+try_chmod bin/mkbundle
+try_chmod bin/mod
+try_chmod bin/mono-api-info
+try_chmod bin/mono-boehm
+try_chmod bin/mono-cil-strip
+try_chmod bin/mono-configuration-crypto
+try_chmod bin/mono-find-provides
+try_chmod bin/mono-find-requires
+try_chmod bin/mono-heapviz
+try_chmod bin/mono-service
+try_chmod bin/mono-service2
+try_chmod bin/mono-shlib-cop
+try_chmod bin/mono-test-install
+try_chmod bin/mono-xmltool
+try_chmod bin/monodocer
+try_chmod bin/monodocs2html
+try_chmod bin/monodocs2slashdoc
+try_chmod bin/monolinker
+try_chmod bin/monop
+try_chmod bin/monop2
+try_chmod bin/mozroots
+try_chmod bin/nunit-console
+try_chmod bin/nunit-console2
+try_chmod bin/nunit-console4
+try_chmod bin/pdb2mdb
+try_chmod bin/permview
+try_chmod bin/peverify
+try_chmod bin/prj2make
+try_chmod bin/resgen
+try_chmod bin/resgen2
+try_chmod bin/secutil
+try_chmod bin/setreg
+try_chmod bin/sgen
+try_chmod bin/signcode
+try_chmod bin/sn
+try_chmod bin/soapsuds
+try_chmod bin/sqlmetal
+try_chmod bin/sqlsharp
+try_chmod bin/svcutil
+try_chmod bin/wsdl
+try_chmod bin/wsdl2
+try_chmod bin/xbuild
+try_chmod bin/xsd
+
+cd "$CUR_DIR"
