Add missing symlink for mono

diff --git a/Engine/Build/BatchFiles/Mac/FixMonoFiles.sh b/Engine/Build/BatchFiles/Mac/FixMonoFiles.sh
--- a/Engine/Build/BatchFiles/Mac/FixMonoFiles.sh
+++ b/Engine/Build/BatchFiles/Mac/FixMonoFiles.sh
@@ -11,6 +11,7 @@ function try_symlink {
 CUR_DIR=`pwd`
 cd ../../../Binaries/ThirdParty/Mono/Mac
 
+try_symlink mono-boehm bin/mono
 try_symlink ../gac/Accessibility/4.0.0.0__b03f5f7f11d50a3a/Accessibility.dll lib/mono/4.5/Accessibility.dll
 try_symlink ../gac/Commons.Xml.Relaxng/4.0.0.0__0738eb9f132ed756/Commons.Xml.Relaxng.dll lib/mono/4.5/Commons.Xml.Relaxng.dll
 try_symlink ../gac/cscompmgd/0.0.0.0__b03f5f7f11d50a3a/cscompmgd.dll lib/mono/4.5/cscompmgd.dll
