Add shell script for Mac batch file which sets up Mono and runs GitDependencies

diff --git a/Engine/Build/BatchFiles/Mac/GitDependencies.sh b/Engine/Build/BatchFiles/Mac/GitDependencies.sh
new file mode 100755
--- /dev/null
+++ b/Engine/Build/BatchFiles/Mac/GitDependencies.sh
@@ -0,0 +1,19 @@
+#!/bin/sh
+# Copyright 1998-2014 Epic Games, Inc. All Rights Reserved.
+
+set -e
+
+cd "`dirname "$0"`"
+
+if [ ! -f ../../../Binaries/DotNET/GitDependencies.exe ]; then
+	echo "Cannot find GitDependencies.exe. This script should be placed in Engine/Build/BatchFiles/Mac."
+	exit 1
+fi 
+
+echo from `pwd`
+
+source SetupMono.sh "`pwd`"
+
+cd ../../../..
+
+mono Engine/Binaries/DotNET/GitDependencies.exe "$@"
