Remove hard-coded list of executables on Mac. GitDependencies will set the executable bit correctly.

diff --git a/Engine/Build/BatchFiles/Mac/FixDependencyFiles.sh b/Engine/Build/BatchFiles/Mac/FixDependencyFiles.sh
--- a/Engine/Build/BatchFiles/Mac/FixDependencyFiles.sh
+++ b/Engine/Build/BatchFiles/Mac/FixDependencyFiles.sh
@@ -1,12 +1,5 @@
 #!/bin/sh
 
-# Fix executable bit for binaries downloaded from GitHub in zip file
-chmod +x ../../../Binaries/Mac/ShaderCompileWorker
-chmod +x ../../../Binaries/Mac/UE4EditorServices.app/Contents/MacOS/UE4EditorServices
-chmod +x ../../../Binaries/Mac/ShaderCompileWorker
-chmod +x ../../../Binaries/Mac/UnrealLightmass
-chmod +x ../../../Binaries/Mac/UnrealPak
-
 # Check for case-sensitive file system
 rm -f casetest*
 touch casetestABC
