Don't try to copy UE4EditorServices.app if it doesn't already exist. Fix for me sourcing SetupMono.sh with 'set -e' when bootstrapping the engine from GitHub.

diff --git a/Engine/Build/BatchFiles/Mac/FixDependencyFiles.sh b/Engine/Build/BatchFiles/Mac/FixDependencyFiles.sh
--- a/Engine/Build/BatchFiles/Mac/FixDependencyFiles.sh
+++ b/Engine/Build/BatchFiles/Mac/FixDependencyFiles.sh
@@ -18,5 +18,7 @@ rm -f casetest*
 
 # Copy UE4EditorServices to ~/Library/Services
 if [ ! -d ~/Library/Services/UE4EditorServices.app ]; then
-	cp -r ../../../Binaries/Mac/UE4EditorServices.app ~/Library/Services/UE4EditorServices.app
+	if [ -d ../../../Binaries/Mac/UE4EditorServices.app ]; then
+		cp -r ../../../Binaries/Mac/UE4EditorServices.app ~/Library/Services/UE4EditorServices.app
+	fi
 fi
