Changes to build process.

- One time setup moved to Setup.sh, with GenerateProjectFiles.sh attempting to detect whether this was ever run previously.
- Checking git/perforce build differently, turned off updating deps.
- Updated smoke test appropriately.
- Closes UEPLAT-231.

diff --git a/Engine/Build/BatchFiles/Linux/BuildThirdParty.sh b/Engine/Build/BatchFiles/Linux/BuildThirdParty.sh
--- a/Engine/Build/BatchFiles/Linux/BuildThirdParty.sh
+++ b/Engine/Build/BatchFiles/Linux/BuildThirdParty.sh
@@ -492,6 +492,9 @@ while :; do
 done
 
 echo "Building ThirdParty libraries"
+echo
+echo "If you don't see SUCCESS message in the end, then building did not finish properly."
+echo "In that case, take a look into ${SCRIPT_DIR}/BuildThirdParty.log for details."
 
 rm -f ${SCRIPT_DIR}/BuildThirdParty.log
 
diff --git a/Engine/Build/BatchFiles/Linux/GenerateProjectFiles.sh b/Engine/Build/BatchFiles/Linux/GenerateProjectFiles.sh
--- a/Engine/Build/BatchFiles/Linux/GenerateProjectFiles.sh
+++ b/Engine/Build/BatchFiles/Linux/GenerateProjectFiles.sh
@@ -1,147 +1,23 @@
 #!/bin/bash
 
 SCRIPT_DIR=$(cd "$(dirname "$BASH_SOURCE")" ; pwd)
-# these need to be passed to children scripts
-export ARCHIVE_ROOT=$HOME/Downloads
-export GITHUB_TAG=latest-preview
-
-IS_GITHUB_BUILD=true
-FORCE_UPDATEDEPS=false
-
-CheckArgs()
-{
-  # if p4 is installed, assume building out of perforce repo (no need to download and fix dependencies)
-  if which p4 > /dev/null; then
-    IS_GITHUB_BUILD=false	# perforce
-  fi
-
-  # Override this with -git
-  for Arg in $@; do
-    if [ "$Arg" == "-git" ]; then
-      IS_GITHUB_BUILD=true
-    elif [ "$Arg" == "-updatedeps" ]; then
-      FORCE_UPDATEDEPS=true     
-    fi
-  done
-}
-
 
 set -e
 
 TOP_DIR=$(cd $SCRIPT_DIR/../../.. ; pwd)
 cd ${TOP_DIR}
 
-if [ ! -d Source ]; then
-  echo "GenerateProjectFiles ERROR: This script file does not appear to be \
-located inside the Engine/Build/BatchFiles/Linux directory."
-  exit 1
-fi
-
-CheckArgs $@
-
-if [ "$(lsb_release --id)" = "Distributor ID:	Ubuntu" -o "$(lsb_release --id)" = "Distributor ID:	Debian" -o "$(lsb_release --id)" = "Distributor ID:	LinuxMint" ]; then
-  # Install all necessary dependencies
-  DEPS="mono-xbuild \
-    mono-dmcs \
-    libmono-microsoft-build-tasks-v4.0-4.0-cil \
-    libmono-system-data-datasetextensions4.0-cil
-    libmono-system-web-extensions4.0-cil
-    libmono-system-management4.0-cil
-    libmono-system-xml-linq4.0-cil
-    libmono-corlib4.0-cil
-    libqt4-dev
-    dos2unix
-    cmake
-    "
-
-  for DEP in $DEPS; do
-    if ! dpkg -s $DEP > /dev/null 2>&1; then
-      echo "Attempting installation of missing package: $DEP"
-      set -x
-      sudo apt-get install -y $DEP
-      set +x
-    fi
-  done
-fi
-
-echo 
-if [ "$IS_GITHUB_BUILD" = true ]; then
-	echo
-	echo Github build
-	echo Checking / downloading the latest archives
+if [ ! -e Build/OneTimeSetupPerformed ]; then
 	echo
-	set +e
-	# should not be needed since we add it as post-merge hook, but run here for now to overkill
-	mono Binaries/DotNET/GitDependencies.exe
-	UpdateResult=$?
-	set -e
-
-	if [ $UpdateResult -eq 1 ]; then
-	  echo
-          echo Could not check/download binary dependencies!
-          exit 1
-        fi
-
-	# check if it had to download anything
-	if [ $UpdateResult -eq 2 -o "$FORCE_UPDATEDEPS" = true ]; then
-	  echo
-          echo Downloaded new binaries!
-	  echo Unpacking and massaging the files
-	  pushd Build/BatchFiles/Linux > /dev/null
-	  ./UpdateDeps.sh 
-	  popd > /dev/null
-        else
-          echo
-          echo All assets are up to date, not unpacking zip files again.
-        fi
-
-else
-	echo Perforce build
-	echo Assuming availability of up to date third-party libraries
+	echo "******************************************************"
+	echo "You have not run Setup.sh, the build will likely fail."
+	echo "******************************************************"
 fi
 
-echo
-pushd Build/BatchFiles/Linux > /dev/null
-./BuildThirdParty.sh
-popd > /dev/null
-
 echo
 echo Setting up Unreal Engine 4 project files...
 echo
 
-# args: wrong filename, correct filename
-# expects to be in Engine folder
-CreateLinkIfNoneExists()
-{
-    WrongName=$1
-    CorrectName=$2
-
-    pushd `dirname $CorrectName` > /dev/null
-    if [ ! -f `basename $CorrectName` ] && [ -f $WrongName ]; then
-      echo "$WrongName -> $CorrectName"
-      ln -sf $WrongName `basename $CorrectName`
-    fi
-    popd > /dev/null
-}
-
-
-
-# Fixes for case sensitive filesystem.
-for BASE in Content/Editor/Slate Content/Slate Documentation/Source/Shared/Icons; do
-  find $BASE -name "*.PNG" | while read PNG_UPPER; do
-    png_lower="$(echo "$PNG_UPPER" | sed 's/.PNG$/.png/')"
-    if [ ! -f $png_lower ]; then
-      PNG_UPPER=$(basename $PNG_UPPER)
-      echo "$png_lower -> $PNG_UPPER"
-      # link, and not move, to make it usable with Perforce workspaces
-      ln -sf `basename "$PNG_UPPER"` "$png_lower"
-    fi
-  done
-done
-
-CreateLinkIfNoneExists ../../engine/shaders/Fxaa3_11.usf  ../Engine/Shaders/Fxaa3_11.usf
-CreateLinkIfNoneExists ../../Engine/shaders/Fxaa3_11.usf  ../Engine/Shaders/Fxaa3_11.usf
-
 set -x
 xbuild Source/Programs/UnrealBuildTool/UnrealBuildTool_Mono.csproj \
   /verbosity:quiet /nologo \
diff --git a/Engine/Build/BatchFiles/Linux/QASmokeManual.sh b/Engine/Build/BatchFiles/Linux/QASmokeManual.sh
--- a/Engine/Build/BatchFiles/Linux/QASmokeManual.sh
+++ b/Engine/Build/BatchFiles/Linux/QASmokeManual.sh
@@ -12,6 +12,16 @@ Header()
   echo "Original location: $LOG_FILE"
 }
 
+OneTimeSetup()
+{
+  echo "Performing one-time setup"
+  set -x
+  cp Engine/Build/Git/Setup.sh .
+  ./Setup.sh
+  rm Setup.sh
+  set +x
+}
+
 GenerateProjectFiles()
 {
   echo "Generating project files"
@@ -70,6 +80,7 @@ if [ -z "$1" ]; then
 
   pushd ../../../.. >> /dev/null
   Run Header
+  Run OneTimeSetup
   Run GenerateProjectFiles
 
   # rebuild dependencies locally for editor targets (extra safety still)
diff --git a/Engine/Build/BatchFiles/Linux/Setup.sh b/Engine/Build/BatchFiles/Linux/Setup.sh
new file mode 100755
--- /dev/null
+++ b/Engine/Build/BatchFiles/Linux/Setup.sh
@@ -0,0 +1,68 @@
+#!/bin/bash
+
+SCRIPT_DIR=$(cd "$(dirname "$BASH_SOURCE")" ; pwd)
+
+set -e
+
+TOP_DIR=$(cd $SCRIPT_DIR/../../.. ; pwd)
+cd ${TOP_DIR}
+
+IS_GITHUB_BUILD=true
+if [ -e Build/PerforceBuild.txt ]; then
+  IS_GITHUB_BUILD=false
+fi
+
+# Ubuntu/Debian/Mint
+if [ "$(lsb_release --id)" = "Distributor ID:	Ubuntu" -o "$(lsb_release --id)" = "Distributor ID:	Debian" -o "$(lsb_release --id)" = "Distributor ID:	LinuxMint" ]; then
+  # Install all necessary dependencies
+  DEPS="mono-xbuild \
+    mono-dmcs \
+    libmono-microsoft-build-tasks-v4.0-4.0-cil \
+    libmono-system-data-datasetextensions4.0-cil
+    libmono-system-web-extensions4.0-cil
+    libmono-system-management4.0-cil
+    libmono-system-xml-linq4.0-cil
+    libmono-corlib4.0-cil
+    libqt4-dev
+    dos2unix
+    cmake
+    "
+
+  for DEP in $DEPS; do
+    if ! dpkg -s $DEP > /dev/null 2>&1; then
+      echo "Attempting installation of missing package: $DEP"
+      set -x
+      sudo apt-get install -y $DEP
+      set +x
+    fi
+  done
+fi
+
+echo 
+if [ "$IS_GITHUB_BUILD" = true ]; then
+	echo
+	echo Github build
+	echo Checking / downloading the latest archives
+	echo
+	set +e
+	mono Binaries/DotNET/GitDependencies.exe --prompt "$@"
+	set -e
+
+	#echo Unpacking and massaging the files
+	#pushd Build/BatchFiles/Linux > /dev/null
+	#./UpdateDeps.sh 
+	#popd > /dev/null
+else
+	echo Perforce build
+	echo Assuming availability of up to date third-party libraries
+fi
+
+echo
+pushd Build/BatchFiles/Linux > /dev/null
+./BuildThirdParty.sh
+popd > /dev/null
+
+touch Build/OneTimeSetupPerformed
+
+
+
diff --git a/Engine/Build/BatchFiles/Linux/UpdateDeps.sh b/Engine/Build/BatchFiles/Linux/UpdateDeps.sh
--- a/Engine/Build/BatchFiles/Linux/UpdateDeps.sh
+++ b/Engine/Build/BatchFiles/Linux/UpdateDeps.sh
@@ -7,27 +7,6 @@
 # This script does not handle building of any of the ThirdParty dependencies.
 # (See BuildThirdParty.sh).
 
-# Location of the archive files (shared with GetAssets.py)
-if [ -z $ARCHIVE_ROOT ]; then
-  echo You should have ARCHIVE_ROOT variable set to the same location as used in GetAssets.py. If you don't know what this is about, don't invoke this file directly.
-fi
-
-# Prefix of the archive files (should match one used by GetAssets.py)
-if [ -z $GITHUB_TAG ]; then
-  echo You should have GITHUB_TAG variable set to the same release tag as used in GetAssets.py. If you don't know what this is about, don't invoke this file directly.
-fi
-
-if [ ! -d $ARCHIVE_ROOT ]; then
-  echo "Download root not found: $ARCHIVE_ROOT"
-  exit 1
-fi
-
-# Zipped archives.
-ARCHIVES="$GITHUB_TAG-Required_1of3.zip
-    $GITHUB_TAG-Required_2of3.zip
-    $GITHUB_TAG-Required_3of3.zip
-    $GITHUB_TAG-Optional.zip"
-
 # Make sure that the current working directory is the root the git checkout,
 # which is four levels up from the this script.
 SCRIPT_DIR=$(cd "$(dirname "$BASH_SOURCE")" ; pwd)
@@ -67,20 +46,6 @@ CleanRepo()
   done
 }
 
-ExtractArchives()
-{
-  echo "==> Installing dependencies"
-  for Archive in $ARCHIVES; do
-    FullPathArchive=$ARCHIVE_ROOT/$Archive
-    if [ ! -f $FullPathArchive ]; then
-      echo "Zip file dependency not found: $FullPathArchive"
-      exit 1
-    fi
-    echo "Extracting $FullPathArchive"
-    unzip -q -n $FullPathArchive
-  done
-}
-
 Unixify()
 {
   echo "==> Converting to unix line endings"
@@ -132,7 +97,6 @@ Unixify()
 }
 
 #CleanRepo
-#ExtractArchives
 Unixify
 
 echo "********** SUCCESS ****************"
diff --git a/Setup.sh b/Setup.sh
--- a/Setup.sh
+++ b/Setup.sh
@@ -1,4 +1,4 @@
-#!/bin/sh
+#!/bin/bash
 # Copyright 1998-2014 Epic Games, Inc. All Rights Reserved.
 
 set -e
@@ -36,5 +36,7 @@ else
 		chmod +x .git/hooks/post-merge
 	fi
 
-	mono Engine/Binaries/DotNET/GitDependencies.exe --prompt "$@"
+	pushd Engine/Build/BatchFiles/Linux > /dev/null
+	./Setup.sh "$@"
+	popd > /dev/null
 fi
