Added fixes for case sensitivity issues.

(missed in previous commit).

diff --git a/Engine/Build/BatchFiles/Linux/Setup.sh b/Engine/Build/BatchFiles/Linux/Setup.sh
--- a/Engine/Build/BatchFiles/Linux/Setup.sh
+++ b/Engine/Build/BatchFiles/Linux/Setup.sh
@@ -2,6 +2,23 @@
 
 SCRIPT_DIR=$(cd "$(dirname "$BASH_SOURCE")" ; pwd)
 
+# args: wrong filename, correct filename
+# expects to be run in Engine folder
+CreateLinkIfNoneExists()
+{
+    WrongName=$1
+    CorrectName=$2
+
+    pushd `dirname $CorrectName` > /dev/null
+    if [ ! -f `basename $CorrectName` ] && [ -f $WrongName ]; then
+      echo "$WrongName -> $CorrectName"
+      ln -sf $WrongName `basename $CorrectName`
+    fi
+    popd > /dev/null
+}
+
+
+# main
 set -e
 
 TOP_DIR=$(cd $SCRIPT_DIR/../../.. ; pwd)
@@ -40,10 +57,8 @@ fi
 
 echo 
 if [ "$IS_GITHUB_BUILD" = true ]; then
-	echo
 	echo Github build
 	echo Checking / downloading the latest archives
-	echo
 	set +e
 	mono Binaries/DotNET/GitDependencies.exe --prompt "$@"
 	set -e
@@ -57,12 +72,26 @@ else
 	echo Assuming availability of up to date third-party libraries
 fi
 
+# Fixes for case sensitive filesystem.
+echo Fixing inconsistent case in filenames.
+for BASE in Content/Editor/Slate Content/Slate Documentation/Source/Shared/Icons; do
+  find $BASE -name "*.PNG" | while read PNG_UPPER; do
+    png_lower="$(echo "$PNG_UPPER" | sed 's/.PNG$/.png/')"
+    if [ ! -f $png_lower ]; then
+      PNG_UPPER=$(basename $PNG_UPPER)
+      echo "$png_lower -> $PNG_UPPER"
+      # link, and not move, to make it usable with Perforce workspaces
+      ln -sf `basename "$PNG_UPPER"` "$png_lower"
+    fi
+  done
+done
+
+CreateLinkIfNoneExists ../../engine/shaders/Fxaa3_11.usf  ../Engine/Shaders/Fxaa3_11.usf
+CreateLinkIfNoneExists ../../Engine/shaders/Fxaa3_11.usf  ../Engine/Shaders/Fxaa3_11.usf
+
 echo
 pushd Build/BatchFiles/Linux > /dev/null
 ./BuildThirdParty.sh
 popd > /dev/null
 
 touch Build/OneTimeSetupPerformed
-
-
-
