Updated Linux README.

diff --git a/Engine/Build/BatchFiles/Linux/README.md b/Engine/Build/BatchFiles/Linux/README.md
--- a/Engine/Build/BatchFiles/Linux/README.md
+++ b/Engine/Build/BatchFiles/Linux/README.md
@@ -1,7 +1,7 @@
 Build scripts for native Linux build
 ====================================
 
-This document describes how to build UE4 4.5 natively on a Linux host. 
+This document describes how to build UE4 4.6 natively on a Linux host. 
 The steps are described here are applicable to the current build, but you may
 want to visit https://wiki.unrealengine.com/Building_On_Linux for the
 latest updates on the process.
@@ -17,7 +17,7 @@ Prerequisites
 -------------
 
 The packages that are required to build the engine vary from distribution to distribution,
-and an up-to-date list should be maintained (and installed) by GenerateProjectFiles.sh -
+and an up-to-date list should be maintained (and installed) by Setup.sh -
 feel free to suggest modifications. Automated install currently works for Debian-based
 distributions only (Debian itself, (K)Ubuntu and Linux Mint).
 
@@ -42,60 +42,55 @@ Setting up/updating the sources
 -------------------------------
 
 Setup has been simplified since the previous releases, and the additional
-.zip archives are now being automatically downloaded. However, in order to
-make this happen, you need to generate a personal access token (repo scope):
-[REDACTED_PROFILE_URL]/tokens/new  Copy the token to OAUTH_TOKEN environment variable -\
-treat it as your GitHub password in terms of security and have that variable set each
-time you are running the below script.
+binary files which are too large to be included into github repository are now being 
+downloaded by GitDependencies tool with minimal hassle. After cloning the repository, 
+you will need to run Setup.sh script which will invoke the said tool to download them.
+The tool will be registered as a post-merge hook, so later updates to binary files
+will be downloaded after each git pull.
 
-More information on github access tokens: https://help.github.com/articles/creating-an-access-token-for-command-line-use/
-
-Step by step:
+How to set up the sources for building, step by step:
 
 1. Clone EpicGames/UnrealEngine repository
 
+    ``git clone https://github.com/EpicGames/UnrealEngine -b 4.6``
+    
+2. Run Setup.sh once.
 
-    git clone https://github.com/EpicGames/UnrealEngine -b 4.5
-
-2. With the aforementioned OAUTH_TOKEN variable set to a valid token,
-run GenerateProjectFiles.sh file in the engine's directory:
-
-
-    cd UnrealEngine
-    ./GenerateProjectFiles.sh
+    ``cd UnrealEngine``
+    
+    ``./Setup.sh``
 
+    The script will try to install additional packages (for certain distributions) and download
+    precompiled binaries of third party libraries. It will also build one of the libraries
+    on your system (LinuxNativeDialogs or LND for short).
 
-The script may ask you to install additional packages (on certain distributions), then it will download 
-the archives with binary dependencies (by default, to ~/Downloads), which is pretty large (3GB in total). 
-They won't be re-downloaded unless they are updated for the particular release tag. 
-If any new archives have been downloaded (which is the case the first time you do this), 
-the script will unpack them over your repository (it will ask you 
-to confirm that action as it is potentially destructive). You can force this by supplying 
--updatedeps switch if for some reason you want to unpack them again.
+    You should see ** SUCCESS ** message after running this step. If you don't, take a look into
+    Troubleshooting section.
+    
+3. After the successful setup, you can generate makefiles (and CMakelists.txt).
 
-The script will also build LinuxNativeDialogs - the only third-party library that currently 
-needs to be built locally.
+    ``./GenerateProjectFiles.sh``
 
-Subsequent updates can be done the same way (remember about OAUTH_TOKEN!)
+Updating the sources later can be done with git pull. The tool to download binary files will be
+registered as a post-merge hook by Setup.sh, so third party libraries will be updated automatically
+(if needed). If you ever need to run it directly, it can be found in Engine/Binaries/DotNET/ directory
+(GitDependencies.exe, which needs to be invoked through mono).
 
-    cd UnrealEngine
-    git pull
-    ./GenerateProjectFiles.sh
 
 
 Building and running
 --------------------
 
 GenerateProjectFiles.sh also produces both makefile and CMakeLists.txt which you can use to import the
-project in your favorite IDE. Caveat: only KDevelop 4.6+ is known to handle the project
-well (it takes about 3-4 GB of resident RAM to load the project).
+project in your favorite IDE. KDevelop 4.6+ is known to handle the project well 
+(it takes about 3-4 GB of resident RAM to load the project though).
 
 The targets match the name of the resulting binary, e.g. UE4Editor-Linux-Debug or UE4Game. You can build them
 by just typing make <target> in the engine's root folder.
 
 Specifically, to be able to run the editor, build the following targets:
 
-    make ShaderCompileWorker UnrealLightmass UE4Editor
+    make ShaderCompileWorker UnrealLightmass UnrealPak UE4Editor
 
 If you intend to develop the editor, you can build a debug configuration of it:
 
@@ -131,6 +126,9 @@ Depending on the project, the editor may need rather large number of file handle
 (e.g. 16000+). If you start seeing errors about not being able to open files, 
 you may need to adjust your limits.
 
+It is advised that you install the editor on a case insensitive filesystem if you
+intend to open projects from other systems (OS X and Windows), or Marketplace.
+
 The time it takes to build the editor in development configuration can be large,
 debug configuration takes about 2/3 of this time. The build process can also take 
 significant amount of RAM (roughly 1GB per core).
