Minor change in Linux readme.

diff --git a/Engine/Build/BatchFiles/Linux/README.md b/Engine/Build/BatchFiles/Linux/README.md
--- a/Engine/Build/BatchFiles/Linux/README.md
+++ b/Engine/Build/BatchFiles/Linux/README.md
@@ -1,7 +1,7 @@
 Build scripts for native Linux build
 ====================================
 
-This document describes how to build UE4 4.6 natively on a Linux host. 
+This document describes how to build Unreal Engine 4.6 natively on a Linux host. 
 The steps are described here are applicable to the current build, but you may
 want to visit https://wiki.unrealengine.com/Building_On_Linux for the
 latest updates on the process.
@@ -65,7 +65,7 @@ How to set up the sources for building, step by step:
     on your system (LinuxNativeDialogs or LND for short).
 
     You should see ** SUCCESS ** message after running this step. If you don't, take a look into
-    Troubleshooting section.
+    BuildThirdParty.log located in Engine/Build/BatchFiles/Linux directory.
     
 3. After the successful setup, you can generate makefiles (and CMakelists.txt).
 
