Added handling of negative UBT exit codes to Mac RocketBuild.sh

diff --git a/Engine/Build/BatchFiles/Mac/RocketBuild.sh b/Engine/Build/BatchFiles/Mac/RocketBuild.sh
--- a/Engine/Build/BatchFiles/Mac/RocketBuild.sh
+++ b/Engine/Build/BatchFiles/Mac/RocketBuild.sh
@@ -55,5 +55,9 @@ case $ACTION in
 		;;
 esac
 
-exit $?
-
+ExitCode=$?
+if [ $ExitCode -eq 254 ] || [ $ExitCode -eq 255 ]; then
+	exit 0
+else
+	exit $ExitCode
+fi
