Updated Mac FixMonoFiles.sh script for Mono 3.10

diff --git a/Engine/Build/BatchFiles/Mac/FixMonoFiles.sh b/Engine/Build/BatchFiles/Mac/FixMonoFiles.sh
--- a/Engine/Build/BatchFiles/Mac/FixMonoFiles.sh
+++ b/Engine/Build/BatchFiles/Mac/FixMonoFiles.sh
@@ -54,8 +54,17 @@ try_symlink ../gac/Mono.Simd/4.0.0.0__0738eb9f132ed756/Mono.Simd.dll lib/mono/4.
 try_symlink ../gac/Mono.Tasklets/4.0.0.0__0738eb9f132ed756/Mono.Tasklets.dll lib/mono/4.5/Mono.Tasklets.dll
 try_symlink ../gac/Mono.Web/4.0.0.0__0738eb9f132ed756/Mono.Web.dll lib/mono/4.5/Mono.Web.dll
 try_symlink ../gac/Mono.WebBrowser/4.0.0.0__0738eb9f132ed756/Mono.WebBrowser.dll lib/mono/4.5/Mono.WebBrowser.dll
+try_symlink ../gac/Mono.XBuild.Tasks/4.0.0.0__0738eb9f132ed756/Mono.XBuild.Tasks.dll lib/mono/4.5/Mono.XBuild.Tasks.dll
 try_symlink ../gac/Novell.Directory.Ldap/4.0.0.0__0738eb9f132ed756/Novell.Directory.Ldap.dll lib/mono/4.5/Novell.Directory.Ldap.dll
 try_symlink ../gac/Npgsql/4.0.0.0__5d8b90d52f46fda7/Npgsql.dll lib/mono/4.5/Npgsql.dll
+try_symlink ../2.0/nunit-console-runner.dll lib/mono/4.5/nunit-console-runner.dll
+try_symlink ../2.0/nunit.core.dll lib/mono/4.5/nunit.core.dll
+try_symlink ../2.0/nunit.core.extensions.dll lib/mono/4.5/nunit.core.extensions.dll
+try_symlink ../2.0/nunit.core.interfaces.dll lib/mono/4.5/nunit.core.interfaces.dll
+try_symlink ../2.0/nunit.framework.dll lib/mono/4.5/nunit.framework.dll
+try_symlink ../2.0/nunit.framework.extensions.dll lib/mono/4.5/nunit.framework.extensions.dll
+try_symlink ../2.0/nunit.mocks.dll lib/mono/4.5/nunit.mocks.dll
+try_symlink ../2.0/nunit.util.dll lib/mono/4.5/nunit.util.dll
 try_symlink ../gac/OpenSystem.C/4.0.0.0__b77a5c561934e089/OpenSystem.C.dll lib/mono/4.5/OpenSystem.C.dll
 try_symlink ../gac/PEAPI/4.0.0.0__0738eb9f132ed756/PEAPI.dll lib/mono/4.5/PEAPI.dll
 try_symlink ../gac/RabbitMQ.Client/4.0.0.0__b03f5f7f11d50a3a/RabbitMQ.Client.dll lib/mono/4.5/RabbitMQ.Client.dll
@@ -91,16 +100,14 @@ try_symlink ../gac/System.Net.Http/4.0.0.0__b03f5f7f11d50a3a/System.Net.Http.dll
 try_symlink ../gac/System.Net.Http.Formatting/4.0.0.0__31bf3856ad364e35/System.Net.Http.Formatting.dll lib/mono/4.5/System.Net.Http.Formatting.dll
 try_symlink ../gac/System.Net.Http.WebRequest/4.0.0.0__b03f5f7f11d50a3a/System.Net.Http.WebRequest.dll lib/mono/4.5/System.Net.Http.WebRequest.dll
 try_symlink ../gac/System.Numerics/4.0.0.0__b77a5c561934e089/System.Numerics.dll lib/mono/4.5/System.Numerics.dll
-try_symlink ../gac/System.Reactive.Core/2.1.30214.0__31bf3856ad364e35/System.Reactive.Core.dll lib/mono/4.5/System.Reactive.Core.dll
-try_symlink ../gac/System.Reactive.Debugger/2.1.30214.0__31bf3856ad364e35/System.Reactive.Debugger.dll lib/mono/4.5/System.Reactive.Debugger.dll
-try_symlink ../gac/System.Reactive.Experimental/2.1.30214.0__31bf3856ad364e35/System.Reactive.Experimental.dll lib/mono/4.5/System.Reactive.Experimental.dll
-try_symlink ../gac/System.Reactive.Interfaces/2.1.30214.0__31bf3856ad364e35/System.Reactive.Interfaces.dll lib/mono/4.5/System.Reactive.Interfaces.dll
-try_symlink ../gac/System.Reactive.Linq/2.1.30214.0__31bf3856ad364e35/System.Reactive.Linq.dll lib/mono/4.5/System.Reactive.Linq.dll
-try_symlink ../gac/System.Reactive.PlatformServices/2.1.30214.0__31bf3856ad364e35/System.Reactive.PlatformServices.dll lib/mono/4.5/System.Reactive.PlatformServices.dll
-try_symlink ../gac/System.Reactive.Providers/2.1.30214.0__31bf3856ad364e35/System.Reactive.Providers.dll lib/mono/4.5/System.Reactive.Providers.dll
-try_symlink ../gac/System.Reactive.Runtime.Remoting/2.1.30214.0__31bf3856ad364e35/System.Reactive.Runtime.Remoting.dll lib/mono/4.5/System.Reactive.Runtime.Remoting.dll
-try_symlink ../gac/System.Reactive.Windows.Forms/2.1.30214.0__31bf3856ad364e35/System.Reactive.Windows.Forms.dll lib/mono/4.5/System.Reactive.Windows.Forms.dll
-try_symlink ../gac/System.Reactive.Windows.Threading/2.1.30214.0__31bf3856ad364e35/System.Reactive.Windows.Threading.dll lib/mono/4.5/System.Reactive.Windows.Threading.dll
+try_symlink ../gac/System.Reactive.Core/2.2.0.0__31bf3856ad364e35/System.Reactive.Core.dll lib/mono/4.5/System.Reactive.Core.dll
+try_symlink ../gac/System.Reactive.Debugger/2.2.0.0__31bf3856ad364e35/System.Reactive.Debugger.dll lib/mono/4.5/System.Reactive.Debugger.dll
+try_symlink ../gac/System.Reactive.Experimental/2.2.0.0__31bf3856ad364e35/System.Reactive.Experimental.dll lib/mono/4.5/System.Reactive.Experimental.dll
+try_symlink ../gac/System.Reactive.Linq/2.2.0.0__31bf3856ad364e35/System.Reactive.Linq.dll lib/mono/4.5/System.Reactive.Linq.dll
+try_symlink ../gac/System.Reactive.Observable.Aliases/0.0.0.0__31bf3856ad364e35/System.Reactive.Observable.Aliases.dll lib/mono/4.5/System.Reactive.Observable.Aliases.dll
+try_symlink ../gac/System.Reactive.PlatformServices/2.2.0.0__31bf3856ad364e35/System.Reactive.PlatformServices.dll lib/mono/4.5/System.Reactive.PlatformServices.dll
+try_symlink ../gac/System.Reactive.Providers/2.2.0.0__31bf3856ad364e35/System.Reactive.Providers.dll lib/mono/4.5/System.Reactive.Providers.dll
+try_symlink ../gac/System.Reactive.Runtime.Remoting/2.2.0.0__31bf3856ad364e35/System.Reactive.Runtime.Remoting.dll lib/mono/4.5/System.Reactive.Runtime.Remoting.dll
 try_symlink ../gac/System.Runtime.Caching/4.0.0.0__b03f5f7f11d50a3a/System.Runtime.Caching.dll lib/mono/4.5/System.Runtime.Caching.dll
 try_symlink ../gac/System.Runtime.DurableInstancing/4.0.0.0__31bf3856ad364e35/System.Runtime.DurableInstancing.dll lib/mono/4.5/System.Runtime.DurableInstancing.dll
 try_symlink ../gac/System.Runtime.Remoting/4.0.0.0__b77a5c561934e089/System.Runtime.Remoting.dll lib/mono/4.5/System.Runtime.Remoting.dll
@@ -132,7 +139,7 @@ try_symlink ../gac/System.Web.WebPages.Deployment/2.0.0.0__31bf3856ad364e35/Syst
 try_symlink ../gac/System.Web.WebPages/2.0.0.0__31bf3856ad364e35/System.Web.WebPages.dll lib/mono/4.5/System.Web.WebPages.dll
 try_symlink ../gac/System.Web.WebPages.Razor/2.0.0.0__31bf3856ad364e35/System.Web.WebPages.Razor.dll lib/mono/4.5/System.Web.WebPages.Razor.dll
 try_symlink ../gac/System.Windows/4.0.0.0__b03f5f7f11d50a3a/System.Windows.dll lib/mono/4.5/System.Windows.dll
-try_symlink ../gac/System.Windows.Forms.DataVisualization/4.0.0.0__b77a5c561934e089/System.Windows.Forms.DataVisualization.dll lib/mono/4.5/System.Windows.Forms.DataVisualization.dll
+try_symlink ../gac/System.Windows.Forms.DataVisualization/4.0.0.0__31bf3856ad364e35/System.Windows.Forms.DataVisualization.dll lib/mono/4.5/System.Windows.Forms.DataVisualization.dll
 try_symlink ../gac/System.Windows.Forms/4.0.0.0__b77a5c561934e089/System.Windows.Forms.dll lib/mono/4.5/System.Windows.Forms.dll
 try_symlink ../gac/System.Xaml/4.0.0.0__b77a5c561934e089/System.Xaml.dll lib/mono/4.5/System.Xaml.dll
 try_symlink ../gac/System.Xml/4.0.0.0__b77a5c561934e089/System.Xml.dll lib/mono/4.5/System.Xml.dll
@@ -140,6 +147,12 @@ try_symlink ../gac/System.Xml.Linq/4.0.0.0__b77a5c561934e089/System.Xml.Linq.dll
 try_symlink ../gac/System.Xml.Serialization/4.0.0.0__b77a5c561934e089/System.Xml.Serialization.dll lib/mono/4.5/System.Xml.Serialization.dll
 try_symlink ../gac/WebMatrix.Data/4.0.0.0__0738eb9f132ed756/WebMatrix.Data.dll lib/mono/4.5/WebMatrix.Data.dll
 try_symlink ../gac/WindowsBase/4.0.0.0__31bf3856ad364e35/WindowsBase.dll lib/mono/4.5/WindowsBase.dll
+try_symlink ../../../gac/Microsoft.Build/12.0.0.0__b03f5f7f11d50a3a/Microsoft.Build.dll lib/mono/xbuild/12.0/bin/Microsoft.Build.dll
+try_symlink ../../../gac/Microsoft.Build.Engine/12.0.0.0__b03f5f7f11d50a3a/Microsoft.Build.Engine.dll lib/mono/xbuild/12.0/bin/Microsoft.Build.Engine.dll
+try_symlink ../../../gac/Microsoft.Build.Framework/12.0.0.0__b03f5f7f11d50a3a/Microsoft.Build.Framework.dll lib/mono/xbuild/12.0/bin/Microsoft.Build.Framework.dll
+try_symlink ../../../gac/Microsoft.Build.Tasks.v12.0/12.0.0.0__b03f5f7f11d50a3a/Microsoft.Build.Tasks.v12.0.dll lib/mono/xbuild/12.0/bin/Microsoft.Build.Tasks.v12.0.dll
+try_symlink ../../../gac/Microsoft.Build.Utilities.v12.0/12.0.0.0__b03f5f7f11d50a3a/Microsoft.Build.Utilities.v12.0.dll lib/mono/xbuild/12.0/bin/Microsoft.Build.Utilities.v12.0.dll
+try_symlink ../../../gac/Mono.XBuild.Tasks/4.0.0.0__0738eb9f132ed756/Mono.XBuild.Tasks.dll lib/mono/xbuild/12.0/bin/Mono.XBuild.Tasks.dll
 if [ ! -d lib/mono/monodoc ]; then
   mkdir lib/mono/monodoc
 fi
@@ -170,6 +183,7 @@ try_chmod bin/gacutil2
 try_chmod bin/genxs
 try_chmod bin/gmcs
 try_chmod bin/httpcfg
+try_chmod bin/idkasm
 try_chmod bin/ilasm
 try_chmod bin/installvst
 try_chmod bin/lc
