Fix error running GitDependencies with new version of mono. Add a symlink to libc.dylib.

diff --git a/Engine/Build/BatchFiles/Mac/GitDependencies.sh b/Engine/Build/BatchFiles/Mac/GitDependencies.sh
--- a/Engine/Build/BatchFiles/Mac/GitDependencies.sh
+++ b/Engine/Build/BatchFiles/Mac/GitDependencies.sh
@@ -19,4 +19,8 @@ source SetupMono.sh "`pwd`"
 
 cd ../../../..
 
+if [ ! -f Engine/Binaries/ThirdParty/Mono/Mac/lib/libmsvcrt.dylib ]; then
+	ln -s /usr/lib/libc.dylib Engine/Binaries/ThirdParty/Mono/Mac/lib/libmsvcrt.dylib
+fi
+
 mono Engine/Binaries/DotNET/GitDependencies.exe "$@"
