- Tweaked Mac/Build.sh to handle IOS as well as iphoneos

diff --git a/Engine/Build/BatchFiles/Mac/Build.sh b/Engine/Build/BatchFiles/Mac/Build.sh
--- a/Engine/Build/BatchFiles/Mac/Build.sh
+++ b/Engine/Build/BatchFiles/Mac/Build.sh
@@ -24,7 +24,7 @@ case $ACTION in
                 Platform=""
                 AdditionalFlags=""
 		
-                case $CLANG_STATIC_ANALYZER_MODE in
+		case $CLANG_STATIC_ANALYZER_MODE in
 				"deep")
 					AdditionalFlags+="-skipActionHistory"
 					;;
@@ -34,8 +34,8 @@ case $ACTION in
 				esac
 
 		case $2 in 
-	                "iphoneos") 
-		                Platform="IOS"
+			"iphoneos"|"IOS")
+		        Platform="IOS"
 				AdditionalFlags+=" -deploy -nocreatestub "
 			;; 
   			"iphonesimulator")
@@ -65,7 +65,7 @@ case $ACTION in
                 AdditionalFlags="-clean"
 
 		case $3 in 
-			"iphoneos")
+			"iphoneos"|"IOS")
 	        	        Platform="IOS"
 				AdditionalFlags+=" -nocreatestub"
 			;;
