Create libmsvcrt.dylib symlink in bundled Mono's lib folder to temporarily workaround some issues with Mono not being able to find it

diff --git a/Engine/Build/BatchFiles/Mac/SetupMono.sh b/Engine/Build/BatchFiles/Mac/SetupMono.sh
--- a/Engine/Build/BatchFiles/Mac/SetupMono.sh
+++ b/Engine/Build/BatchFiles/Mac/SetupMono.sh
@@ -28,6 +28,10 @@ if [ $IS_MONO_INSTALLED -eq 0 ]; then
 	export PATH=$UE_MONO_DIR/bin:$PATH
 	export MONO_PATH=$UE_MONO_DIR/lib:$MONO_PATH
 	export LD_LIBRARY_PATH=$UE_MONO_DIR/lib:$LD_LIBRARY_PATH
+
+	if [ ! -f $UE_MONO_DIR/lib/libmsvcrt.dylib ]; then
+		ln -s /usr/lib/libc.dylib $UE_MONO_DIR/lib/libmsvcrt.dylib
+	fi
 fi
 
 cd "$START_DIR"
