When setting up mono on the Mac use 'which mono' to determine the location of the system version of Mono. This fixes [UE-7341] "Use 'which mono' in SetupMono.sh to determine path to system mono".

diff --git a/Engine/Build/BatchFiles/Mac/SetupMono.sh b/Engine/Build/BatchFiles/Mac/SetupMono.sh
--- a/Engine/Build/BatchFiles/Mac/SetupMono.sh
+++ b/Engine/Build/BatchFiles/Mac/SetupMono.sh
@@ -5,11 +5,12 @@ sh FixMonoFiles.sh
 sh FixDependencyFiles.sh
 
 IS_MONO_INSTALLED=0
-if [ -f "/usr/bin/mono" ]; then
+MONO_VERSION_PATH=`which mono`
+if [ -f $MONO_VERSION_PATH ]; then
 	# If Mono is installed, check if it's 3.2.6 or higher
 	MONO_VERSION_PREFIX="Mono JIT compiler version "
 	MONO_VERSION_PREFIX_LEN=${#MONO_VERSION_PREFIX}
-	MONO_VERSION=`/usr/bin/mono --version |grep "$MONO_VERSION_PREFIX"`
+	MONO_VERSION=`"${MONO_VERSION_PATH}" --version |grep "$MONO_VERSION_PREFIX"`
 	MONO_VERSION=(`echo ${MONO_VERSION:MONO_VERSION_PREFIX_LEN} |tr '.' ' '`)
 	if [ ${MONO_VERSION[0]} -ge 3 ]; then
 		if [ ${MONO_VERSION[1]} -eq 2 ] && [ ${MONO_VERSION[2]} -ge 6 ]; then
