Fixed Mono detection using which tool for a case when Mono is not installed

diff --git a/Engine/Build/BatchFiles/Mac/SetupMono.sh b/Engine/Build/BatchFiles/Mac/SetupMono.sh
--- a/Engine/Build/BatchFiles/Mac/SetupMono.sh
+++ b/Engine/Build/BatchFiles/Mac/SetupMono.sh
@@ -6,7 +6,7 @@ sh FixDependencyFiles.sh
 
 IS_MONO_INSTALLED=0
 MONO_VERSION_PATH=`which mono`
-if [ -f $MONO_VERSION_PATH ]; then
+if [ ! $MONO_VERSION_PATH == "" ] && [ -f $MONO_VERSION_PATH ]; then
 	# If Mono is installed, check if it's 3.2.6 or higher
 	MONO_VERSION_PREFIX="Mono JIT compiler version "
 	MONO_VERSION_PREFIX_LEN=${#MONO_VERSION_PREFIX}
