Linux readme updates (UEPLAT-480).

diff --git a/Engine/Build/BatchFiles/Linux/README.md b/Engine/Build/BatchFiles/Linux/README.md
--- a/Engine/Build/BatchFiles/Linux/README.md
+++ b/Engine/Build/BatchFiles/Linux/README.md
@@ -1,7 +1,7 @@
 Build scripts for native Linux build
 ====================================
 
-This document describes how to build Unreal Engine 4.6 natively on a Linux host. 
+This document describes how to build Unreal Engine 4.7 natively on a Linux host. 
 The steps are described here are applicable to the current build, but you may
 want to visit https://wiki.unrealengine.com/Building_On_Linux for the
 latest updates on the process.
@@ -23,8 +23,7 @@ distributions only (Debian itself, (K)Ubuntu and Linux Mint).
 
 Most important dependencies:
 - mono 3.x (2.x may work, but is not recommended), including xbuild and C# compiler (*mcs), and libraries for NET 4.0 framework.
-- clang 3.3 (clang 3.5.0 is also supported, but NOT clang 3.4).
-- python (2 or 3) - needed for the script that downloads the binary dependencies.
+- clang 3.5.0 (clang 3.3 is also supported, but NOT clang 3.4).
 - Qt(4,5) or GTK development packages to build LinuxNativeDialogs.
 - SDL2 is also needed for building LinuxNativeDialogs module, but the rest
 of the engine is using our own (modified) version of it from Engine/Source/ThirdParty/SDL2.
@@ -52,7 +51,7 @@ How to set up the sources for building, step by step:
 
 1. Clone EpicGames/UnrealEngine repository
 
-    ``git clone https://github.com/EpicGames/UnrealEngine -b 4.6``
+    ``git clone https://github.com/EpicGames/UnrealEngine -b 4.7``
     
 2. Run Setup.sh once.
 
@@ -118,6 +117,11 @@ You can also append -game if you want to run the project as a game (you can also
 Notes
 -----
 
+Depending on mono version and some other not yet clarified circumstances (this may be relevant: 
+http://stackoverflow.com/questions/13859467/ravendb-client-onlinux-connecting-to-windows-server-using-mono-http),
+binary downloader tool invoked by Setup.sh may fail. In that case, re-run Setup.sh (repeat step #2
+of the above install process).
+
 On the first start, the editor will be "compiling shaders" (a bit misleading terminology, 
 it will be converting them to GLSL). The result will be stored in Engine/DerivedDataCache folder 
 and used for subsequent runs.
