Moved wrapper script for GitDependencies to a better location.

- Suggested by BenM.

diff --git a/Engine/Build/BatchFiles/Linux/GitDependencies.sh b/Engine/Build/BatchFiles/Linux/GitDependencies.sh
new file mode 100755
--- /dev/null
+++ b/Engine/Build/BatchFiles/Linux/GitDependencies.sh
@@ -0,0 +1,18 @@
+#!/bin/bash
+
+# Certain versions of mono can crash while running GitDependencies
+# This script checks if that happened and simply re-runs the tool.
+
+ARGS=$@
+# cd to Engine root
+cd "$(dirname "$BASH_SOURCE")"/../../../..
+
+while : ; do
+        mono Engine/Binaries/DotNET/GitDependencies.exe $ARGS
+        RESULT=$?
+
+        echo "Result: $RESULT"
+        # quit if not crashed
+        [[ RESULT -eq 139 ]] || break
+        echo "mono GitDependencies.exe $ARGS crashed with return code $RESULT" >> GitDependencies.crash.log
+done
diff --git a/Engine/Build/BatchFiles/Linux/Setup.sh b/Engine/Build/BatchFiles/Linux/Setup.sh
--- a/Engine/Build/BatchFiles/Linux/Setup.sh
+++ b/Engine/Build/BatchFiles/Linux/Setup.sh
@@ -60,7 +60,7 @@ if [ "$IS_GITHUB_BUILD" = true ]; then
 	echo Github build
 	echo Checking / downloading the latest archives
 	set +e
-	mono Binaries/DotNET/GitDependencies.exe --prompt "$@"
+	Build/BatchFiles/Linux/GitDependencies.sh --prompt "$@"
 	set -e
 
 	#echo Unpacking and massaging the files
diff --git a/Setup.sh b/Setup.sh
--- a/Setup.sh
+++ b/Setup.sh
@@ -28,11 +28,11 @@ else
 	if [ -d .git/hooks ]; then
 		echo "Registering git hooks... (this will override existing ones!)"
 		echo \#!/bin/sh >.git/hooks/post-checkout
-		echo Engine/Binaries/DotNET/GitDependencies.sh >>.git/hooks/post-checkout
+		echo Engine/Build/BatchFiles/Linux/GitDependencies.sh >>.git/hooks/post-checkout
 		chmod +x .git/hooks/post-checkout
 
 		echo \#!/bin/sh >.git/hooks/post-merge
-		echo Engine/Binaries/DotNET/GitDependencies.sh >>.git/hooks/post-merge
+		echo Engine/Build/BatchFiles/Linux/GitDependencies.sh >>.git/hooks/post-merge
 		chmod +x .git/hooks/post-merge
 	fi
 
