Automagically install dependecies on other debain like distros.

- PR #768 (UE-8171) contributed by [REDACTED_PERSON], with changes.

diff --git a/Engine/Build/BatchFiles/Linux/Setup.sh b/Engine/Build/BatchFiles/Linux/Setup.sh
--- a/Engine/Build/BatchFiles/Linux/Setup.sh
+++ b/Engine/Build/BatchFiles/Linux/Setup.sh
@@ -29,30 +29,33 @@ if [ -e Build/PerforceBuild.txt ]; then
   IS_GITHUB_BUILD=false
 fi
 
-# Ubuntu/Debian/Mint
-if [ "$(lsb_release --id)" = "Distributor ID:	Ubuntu" -o "$(lsb_release --id)" = "Distributor ID:	Debian" -o "$(lsb_release --id)" = "Distributor ID:	LinuxMint" ]; then
-  # Install all necessary dependencies
-  DEPS="mono-xbuild \
-    mono-dmcs \
-    libmono-microsoft-build-tasks-v4.0-4.0-cil \
-    libmono-system-data-datasetextensions4.0-cil
-    libmono-system-web-extensions4.0-cil
-    libmono-system-management4.0-cil
-    libmono-system-xml-linq4.0-cil
-    libmono-corlib4.0-cil
-    libqt4-dev
-    dos2unix
-    cmake
-    "
+if [ -e /etc/os-release ]; then
+  source /etc/os-release
+  # Ubuntu/Debian/Mint
+  if [[ "$ID" == "ubuntu" ]] || [[ "$ID_LIKE" == "ubuntu" ]] || [ "$ID" == "debian" ]] || [[ "$ID_LIKE" == "debian" ]]; then
+    # Install all necessary dependencies
+    DEPS="mono-xbuild \
+      mono-dmcs \
+      libmono-microsoft-build-tasks-v4.0-4.0-cil \
+      libmono-system-data-datasetextensions4.0-cil
+      libmono-system-web-extensions4.0-cil
+      libmono-system-management4.0-cil
+      libmono-system-xml-linq4.0-cil
+      libmono-corlib4.0-cil
+      libqt4-dev
+      dos2unix
+      cmake
+      "
 
-  for DEP in $DEPS; do
-    if ! dpkg -s $DEP > /dev/null 2>&1; then
-      echo "Attempting installation of missing package: $DEP"
-      set -x
-      sudo apt-get install -y $DEP
-      set +x
-    fi
-  done
+    for DEP in $DEPS; do
+      if ! dpkg -s $DEP > /dev/null 2>&1; then
+        echo "Attempting installation of missing package: $DEP"
+        set -x
+        sudo apt-get install -y $DEP
+        set +x
+      fi
+    done
+  fi
 fi
 
 echo 
