Updated Linux readme after UE-7657 resolution.

diff --git a/Engine/Build/BatchFiles/Linux/README.md b/Engine/Build/BatchFiles/Linux/README.md
--- a/Engine/Build/BatchFiles/Linux/README.md
+++ b/Engine/Build/BatchFiles/Linux/README.md
@@ -119,8 +119,9 @@ Notes
 
 Depending on mono version and some other not yet clarified circumstances (this may be relevant: 
 http://stackoverflow.com/questions/13859467/ravendb-client-onlinux-connecting-to-windows-server-using-mono-http),
-binary downloader tool invoked by Setup.sh may fail. In that case, re-run Setup.sh (repeat step #2
-of the above install process).
+binary downloader tool invoked by Setup.sh may fail. In that case, Setup.sh will keep re-running it
+until it succeeds (or at least stops crashing). You may want to keep an eye on this as there is
+a slight possibility of the script getting stuck in an infinite loop if the tool keeps crashing.
 
 On the first start, the editor will be "compiling shaders" (a bit misleading terminology, 
 it will be converting them to GLSL). The result will be stored in Engine/DerivedDataCache folder 
