Force a clean rebuild of UBT when generating project files from GenerateProjectFiles.bat. The UBT project conditionally includes PS4/XboxOne source files, but does not detect the presence/absence of those files when dependency checking. As a result, a user who copies in PS4/XboxOne files after building UBT before does not get the behavior they would expect until they force a rebuild manually (missing targets, not using 2012 projects for XboxOne, etc...).

diff --git a/Engine/Build/BatchFiles/GenerateProjectFiles.bat b/Engine/Build/BatchFiles/GenerateProjectFiles.bat
--- a/Engine/Build/BatchFiles/GenerateProjectFiles.bat
+++ b/Engine/Build/BatchFiles/GenerateProjectFiles.bat
@@ -55,7 +55,7 @@ goto Error_NoVisualStudioEnvironment
 
 
 :ReadyToCompile
-msbuild /nologo /verbosity:quiet Programs\UnrealBuildTool\UnrealBuildTool.csproj /property:Configuration=Development /property:Platform=AnyCPU
+msbuild /nologo /verbosity:quiet Programs\UnrealBuildTool\UnrealBuildTool.csproj /property:Configuration=Development /property:Platform=AnyCPU /target:Clean,Build
 if not %ERRORLEVEL% == 0 goto Error_UBTCompileFailed
 
 rem ## Run UnrealBuildTool to generate Visual Studio solution and project files
