Fix for early exit when running Git setup script if mono is not installed. "set -e" from setup.sh causes failure when capturing "where mono".

diff --git a/Engine/Build/BatchFiles/Mac/SetupMono.sh b/Engine/Build/BatchFiles/Mac/SetupMono.sh
--- a/Engine/Build/BatchFiles/Mac/SetupMono.sh
+++ b/Engine/Build/BatchFiles/Mac/SetupMono.sh
@@ -5,7 +5,7 @@ sh FixMonoFiles.sh
 sh FixDependencyFiles.sh
 
 IS_MONO_INSTALLED=0
-MONO_VERSION_PATH=`which mono`
+MONO_VERSION_PATH=`which mono` || true
 if [ ! $MONO_VERSION_PATH == "" ] && [ -f $MONO_VERSION_PATH ]; then
 	# If Mono is installed, check if it's 3.2.6 or higher
 	MONO_VERSION_PREFIX="Mono JIT compiler version "
