Mono/xbuild: force tools version to be 4.0 (fixes UE-8669).

diff --git a/Engine/Build/BatchFiles/Linux/GenerateProjectFiles.sh b/Engine/Build/BatchFiles/Linux/GenerateProjectFiles.sh
--- a/Engine/Build/BatchFiles/Linux/GenerateProjectFiles.sh
+++ b/Engine/Build/BatchFiles/Linux/GenerateProjectFiles.sh
@@ -26,27 +26,32 @@ xbuild Source/Programs/UnrealBuildTool/UnrealBuildTool.csproj \
 
 xbuild Source/Programs/AutomationTool/AutomationTool.csproj \
   /verbosity:quiet /nologo \
+  /tv:4.0 \
   /p:TargetFrameworkVersion=v4.0 \
   /p:Platform="AnyCPU" \
   /p:Configuration="Development"
 
 xbuild Source/Programs/AutomationTool/Scripts/AutomationScripts.Automation.csproj \
   /verbosity:quiet /nologo \
+  /tv:4.0 \
   /p:TargetFrameworkVersion=v4.0 \
   /p:Configuration="Development"
 
 xbuild Source/Programs/AutomationTool/Linux/Linux.Automation.csproj \
   /verbosity:quiet /nologo \
+  /tv:4.0 \
   /p:TargetFrameworkVersion=v4.0 \
   /p:Configuration="Development"
 
 xbuild Source/Programs/AutomationTool/Android/Android.Automation.csproj \
   /verbosity:quiet /nologo \
+  /tv:4.0 \
   /p:TargetFrameworkVersion=v4.0 \
   /p:Configuration="Development"
 
 xbuild Source/Programs/AutomationTool/HTML5/HTML5.Automation.csproj \
   /verbosity:quiet /nologo \
+  /tv:4.0 \
   /p:TargetFrameworkVersion=v4.0 \
   /p:Configuration="Development"
 
