More robust restarting of GitDependencies.exe on Linux.

- Mono can crash with SIGABRT etc.
- Fixes UE-7657

diff --git a/Engine/Build/BatchFiles/Linux/GitDependencies.sh b/Engine/Build/BatchFiles/Linux/GitDependencies.sh
--- a/Engine/Build/BatchFiles/Linux/GitDependencies.sh
+++ b/Engine/Build/BatchFiles/Linux/GitDependencies.sh
@@ -13,6 +13,6 @@ while : ; do
 
         echo "Result: $RESULT"
         # quit if not crashed
-        [[ RESULT -eq 139 ]] || break
+        [[ $RESULT -lt 129 ]] && break
         echo "mono GitDependencies.exe $ARGS crashed with return code $RESULT" >> GitDependencies.crash.log
 done
