Linux: support autoinstalling dependencies on SUSE

- PR #813 contributed by [REDACTED_PERSON].
- Not tested by us since we don't use SUSE.

diff --git a/Engine/Build/BatchFiles/Linux/Setup.sh b/Engine/Build/BatchFiles/Linux/Setup.sh
--- a/Engine/Build/BatchFiles/Linux/Setup.sh
+++ b/Engine/Build/BatchFiles/Linux/Setup.sh
@@ -56,6 +56,26 @@ if [ -e /etc/os-release ]; then
       fi
     done
   fi
+  
+  # openSUSE/SLED/SLES
+  if [[ "$ID" == "opensuse" ]] || [[ "$ID_LIKE" == "suse" ]]; then
+    # Install all necessary dependencies
+    DEPS="mono-core
+      mono-devel
+      libqt4-devel
+      dos2unix
+      cmake
+      "
+
+    for DEP in $DEPS; do
+      if ! rpm -q $DEP > /dev/null 2>&1; then
+        echo "Attempting installation of missing package: $DEP"
+        set -x
+        sudo zypper -n install $DEP
+        set +x
+      fi
+    done
+  fi
 fi
 
 echo 
