Copy PS4UE4.natvis to the correct folder for PS4 debugging visualizers.

diff --git a/Engine/Build/BatchFiles/CopyVisualizers.bat b/Engine/Build/BatchFiles/CopyVisualizers.bat
--- a/Engine/Build/BatchFiles/CopyVisualizers.bat
+++ b/Engine/Build/BatchFiles/CopyVisualizers.bat
@@ -9,7 +9,7 @@ if errorlevel 0 (
 		pushd %~dp0
 		call GetVSComnToolsPath 12
 		popd
-		
+	
 		if "!VsComnToolsPath!" == "" goto NoVisualStudio2013Environment
 		attrib -R "!UE4_MyDocs!\Visual Studio 2013\Visualizers\UE4.natvis" 1>nul 2>nul
 		copy /Y "%~dp0..\..\Extras\VisualStudioDebugging\UE4.natvis" "!UE4_MyDocs!\Visual Studio 2013\Visualizers\UE4.natvis" 1>nul 2>nul
@@ -24,6 +24,11 @@ if errorlevel 0 (
 		copy /Y "%~dp0..\..\Extras\VisualStudioDebugging\UE4.natvis" "!UE4_MyDocs!\Visual Studio 2012\Visualizers\UE4.natvis" 1>nul 2>nul
 
 	:NoVisualStudio2012Environment
+		if "%SCE_ORBIS_SDK_DIR%" == "" goto NoPS4Environment
+		attrib -R "!UE4_MyDocs!\SCE\orbis-debugger\UE4autoexp.dat" 1>nul 2>nul
+		copy /Y "%~dp0..\..\Extras\VisualStudioDebugging\PS4UE4.natvis" "!UE4_MyDocs!\SCE\orbis-debugger\PS4UE4.natvis" 1>nul 2>nul
+
+	:NoPS4Environment
 		set UE4_MyDocs=
 	)
 )
