Fix attrib to correct file

diff --git a/Engine/Build/BatchFiles/CopyVisualizers.bat b/Engine/Build/BatchFiles/CopyVisualizers.bat
--- a/Engine/Build/BatchFiles/CopyVisualizers.bat
+++ b/Engine/Build/BatchFiles/CopyVisualizers.bat
@@ -25,7 +25,7 @@ if errorlevel 0 (
 
 	:NoVisualStudio2012Environment
 		if "%SCE_ORBIS_SDK_DIR%" == "" goto NoPS4Environment
-		attrib -R "!UE4_MyDocs!\SCE\orbis-debugger\UE4autoexp.dat" 1>nul 2>nul
+		attrib -R "!UE4_MyDocs!\SCE\orbis-debugger\PS4UE4.natvis" 1>nul 2>nul
 		copy /Y "%~dp0..\..\Extras\VisualStudioDebugging\PS4UE4.natvis" "!UE4_MyDocs!\SCE\orbis-debugger\PS4UE4.natvis" 1>nul 2>nul
 
 	:NoPS4Environment
