PR #872: Move UE4EditorServices.app service installing code to run after it is built. (Contributed by [REDACTED_PERSON])

https://github.com/EpicGames/UnrealEngine/pull/872

diff --git a/Engine/Build/BatchFiles/Mac/GitDependencies.sh b/Engine/Build/BatchFiles/Mac/GitDependencies.sh
--- a/Engine/Build/BatchFiles/Mac/GitDependencies.sh
+++ b/Engine/Build/BatchFiles/Mac/GitDependencies.sh
@@ -24,3 +24,7 @@ if [ ! -f Engine/Binaries/ThirdParty/Mono/Mac/lib/libmsvcrt.dylib ]; then
 fi
 
 mono Engine/Binaries/DotNET/GitDependencies.exe "$@"
+
+pushd $(dirname $SCRIPT_PATH) > /dev/null
+sh FixDependencyFiles.sh
+popd > /dev/null
