Add check for dependencies on Arch Linux in Setup.sh

- PR #880 contributed by [REDACTED_PERSON].
- Unsupported/untested, but included to make Arch users experience better.

https://github.com/EpicGames/UnrealEngine/pull/880

diff --git a/Engine/Build/BatchFiles/Linux/Setup.sh b/Engine/Build/BatchFiles/Linux/Setup.sh
--- a/Engine/Build/BatchFiles/Linux/Setup.sh
+++ b/Engine/Build/BatchFiles/Linux/Setup.sh
@@ -32,7 +32,7 @@ fi
 if [ -e /etc/os-release ]; then
   source /etc/os-release
   # Ubuntu/Debian/Mint
-  if [[ "$ID" == "ubuntu" ]] || [[ "$ID_LIKE" == "ubuntu" ]] || [ "$ID" == "debian" ]] || [[ "$ID_LIKE" == "debian" ]]; then
+  if [[ "$ID" == "ubuntu" ]] || [[ "$ID_LIKE" == "ubuntu" ]] || [[ "$ID" == "debian" ]] || [[ "$ID_LIKE" == "debian" ]]; then
     # Install all necessary dependencies
     DEPS="mono-xbuild \
       mono-dmcs \
@@ -76,6 +76,24 @@ if [ -e /etc/os-release ]; then
       fi
     done
   fi
+
+  # Arch Linux
+  if [[ "$ID" == "arch" ]] || [[ "$ID_LIKE" == "arch" ]]; then
+    DEPS="clang mono python sdl2 qt4 dos2unix cmake"
+    MISSING=false
+    for DEP in $DEPS; do
+      if ! pacman -Qs $DEP > /dev/null 2>&1; then
+        MISSING=true
+        break
+      fi
+    done
+    if [ "$MISSING" = true ]; then
+      echo "Attempting to install missing packages: $DEPS"
+      set -x
+      sudo pacman -Sy --needed --noconfirm $DEPS
+      set +x
+    fi
+  fi
 fi
 
 echo 
