Propagate return code from GitDependencies.exe

- Makes Setup.sh fail in case GitDependencies fail, avoiding confusion.
- See https://forums.unrealengine.com/showthread.php?51998-GitHub-dependency-zips-and-latest-preview-snapshots for details.

diff --git a/Engine/Build/BatchFiles/Linux/GitDependencies.sh b/Engine/Build/BatchFiles/Linux/GitDependencies.sh
--- a/Engine/Build/BatchFiles/Linux/GitDependencies.sh
+++ b/Engine/Build/BatchFiles/Linux/GitDependencies.sh
@@ -6,6 +6,7 @@
 ARGS=$@
 # cd to Engine root
 cd "$(dirname "$BASH_SOURCE")"/../../../..
+RESULT=0
 
 while : ; do
         mono Engine/Binaries/DotNET/GitDependencies.exe $ARGS
@@ -16,3 +17,5 @@ while : ; do
         [[ $RESULT -lt 129 ]] && break
         echo "mono GitDependencies.exe $ARGS crashed with return code $RESULT" >> GitDependencies.crash.log
 done
+
+exit $RESULT
diff --git a/Engine/Build/BatchFiles/Linux/Setup.sh b/Engine/Build/BatchFiles/Linux/Setup.sh
--- a/Engine/Build/BatchFiles/Linux/Setup.sh
+++ b/Engine/Build/BatchFiles/Linux/Setup.sh
@@ -100,14 +100,7 @@ echo
 if [ "$IS_GITHUB_BUILD" = true ]; then
 	echo Github build
 	echo Checking / downloading the latest archives
-	set +e
 	Build/BatchFiles/Linux/GitDependencies.sh --prompt "$@"
-	set -e
-
-	#echo Unpacking and massaging the files
-	#pushd Build/BatchFiles/Linux > /dev/null
-	#./UpdateDeps.sh 
-	#popd > /dev/null
 else
 	echo Perforce build
 	echo Assuming availability of up to date third-party libraries
