- Xcode unlock

diff --git a/Engine/Build/BatchFiles/Mac/xcodeunlock.sh b/Engine/Build/BatchFiles/Mac/xcodeunlock.sh
--- a/Engine/Build/BatchFiles/Mac/xcodeunlock.sh
+++ b/Engine/Build/BatchFiles/Mac/xcodeunlock.sh
@@ -24,7 +24,7 @@ to=10
 ##  when unlocking a file from Xcode.
 ## If set to 0, the file is unlocked and checked out of Perforce
 ## silently (messages and errors will still be logged to console):
-fb=1
+fb=0
 
 # Xcode doesn't setup the environment, so we load it here:
 source ~/.bash_profile
