Linux: add missing dependency for System.IO.Packaging.

- Fixes compilation of some Rocket projects (UE-12473).

diff --git a/Engine/Build/BatchFiles/Linux/Setup.sh b/Engine/Build/BatchFiles/Linux/Setup.sh
--- a/Engine/Build/BatchFiles/Linux/Setup.sh
+++ b/Engine/Build/BatchFiles/Linux/Setup.sh
@@ -42,6 +42,7 @@ if [ -e /etc/os-release ]; then
       libmono-system-management4.0-cil
       libmono-system-xml-linq4.0-cil
       libmono-corlib4.0-cil
+      libmono-windowsbase4.0-cil
       clang-3.5
       "
 
