Linux: remove stale TPS library that was moved in 4.8 (UE-16033).

- Removed when running Setup.sh
- Addresses issues like https://answers.unrealengine.com/questions/231689/linux-error-when-run-editor-undefined-symbol-ulinu.html

diff --git a/Engine/Build/BatchFiles/Linux/Setup.sh b/Engine/Build/BatchFiles/Linux/Setup.sh
--- a/Engine/Build/BatchFiles/Linux/Setup.sh
+++ b/Engine/Build/BatchFiles/Linux/Setup.sh
@@ -129,6 +129,9 @@ done
 CreateLinkIfNoneExists ../../engine/shaders/Fxaa3_11.usf  ../Engine/Shaders/Fxaa3_11.usf
 CreateLinkIfNoneExists ../../Engine/shaders/Fxaa3_11.usf  ../Engine/Shaders/Fxaa3_11.usf
 
+echo Removing a stable libLND.so binary that was relocated in 4.8
+rm -f ../Engine/Binaries/Linux/libLND.so
+
 echo
 pushd Build/BatchFiles/Linux > /dev/null
 ./BuildThirdParty.sh
