Linux: Mono version check and adding Fedora deps.

- PR #1263: Check minimum Mono version in Linux/Setup.sh (Contributed by [REDACTED_PERSON])
- PR #1261: Add Fedora Linux support to Setup.sh (Contributed by [REDACTED_PERSON])

diff --git a/Engine/Build/BatchFiles/Linux/Setup.sh b/Engine/Build/BatchFiles/Linux/Setup.sh
--- a/Engine/Build/BatchFiles/Linux/Setup.sh
+++ b/Engine/Build/BatchFiles/Linux/Setup.sh
@@ -83,6 +83,27 @@ if [ -e /etc/os-release ]; then
     done
   fi
 
+  # Fedora
+  if [[ "$ID" == "fedora" ]]; then
+    # Install all necessary dependencies
+    DEPS="mono-core
+      mono-devel
+      qt-devel
+      dos2unix
+      cmake
+      "
+
+    for DEP in $DEPS; do
+      if ! rpm -q $DEP > /dev/null 2>&1; then
+        echo "Attempting installation of missing package: $DEP"
+        set -x
+        sudo dnf -y install $DEP
+        set +x
+      fi
+    done
+  fi
+
+
   # Arch Linux
   if [[ "$ID" == "arch" ]] || [[ "$ID_LIKE" == "arch" ]]; then
     DEPS="clang mono python sdl2 qt4 dos2unix cmake"
@@ -102,6 +123,14 @@ if [ -e /etc/os-release ]; then
   fi
 fi
 
+MONO_MINIMUM_VERSION=3
+MONO_VERSION=$(mono -V | sed -n 's/.* version \([0-9]\+\).*/\1/p')
+if [[ "$MONO_VERSION" -lt "$MONO_MINIMUM_VERSION" ]]; then
+  echo "Minimum required Mono version is $MONO_MINIMUM_VERSION. Installed is:"
+  mono -V | sed -n '/version/p'
+  exit 1
+fi
+
 echo 
 if [ "$IS_GITHUB_BUILD" = true ]; then
 	echo Github build
