Linux: PR #1240: Don't do partial upgrades on Arch Linux (UE-16919)

- Contributed by [REDACTED_PERSON], merged as is without testing.

diff --git a/Engine/Build/BatchFiles/Linux/Setup.sh b/Engine/Build/BatchFiles/Linux/Setup.sh
--- a/Engine/Build/BatchFiles/Linux/Setup.sh
+++ b/Engine/Build/BatchFiles/Linux/Setup.sh
@@ -117,7 +117,7 @@ if [ -e /etc/os-release ]; then
     if [ "$MISSING" = true ]; then
       echo "Attempting to install missing packages: $DEPS"
       set -x
-      sudo pacman -Sy --needed --noconfirm $DEPS
+      sudo pacman -S --needed --noconfirm $DEPS
       set +x
     fi
   fi
