PR #1231: Fixed dependency script for Mac. (Contributed by [REDACTED_PERSON])

diff --git a/Engine/Build/BatchFiles/Mac/GitDependencies.sh b/Engine/Build/BatchFiles/Mac/GitDependencies.sh
--- a/Engine/Build/BatchFiles/Mac/GitDependencies.sh
+++ b/Engine/Build/BatchFiles/Mac/GitDependencies.sh
@@ -8,7 +8,7 @@ if [ -h $SCRIPT_PATH ]; then
 	SCRIPT_PATH=$(dirname $SCRIPT_PATH)/$(readlink $SCRIPT_PATH)
 fi
 
-cd $(dirname $SCRIPT_PATH)
+cd $(dirname $SCRIPT_PATH) && SCRIPT_PATH="`pwd`/$(basename $SCRIPT_PATH)"
 
 if [ ! -f ../../../Binaries/DotNET/GitDependencies.exe ]; then
 	echo "Cannot find GitDependencies.exe. This script should be placed in Engine/Build/BatchFiles/Mac."
