UEB-120 - GenerateProjectFiles.bat will now exit with a non-zero exit code (1) when any failure happens. This should allow UGS to properly detect an error when, say, VS is not installed. Tested on a machine without VS installed and seems to work.

diff --git a/Engine/Build/BatchFiles/GenerateProjectFiles.bat b/Engine/Build/BatchFiles/GenerateProjectFiles.bat
--- a/Engine/Build/BatchFiles/GenerateProjectFiles.bat
+++ b/Engine/Build/BatchFiles/GenerateProjectFiles.bat
@@ -64,7 +64,8 @@ rem ## NOTE: We also pass along any arguments to the GenerateProjectFiles.bat he
 if not %ERRORLEVEL% == 0 goto Error_ProjectGenerationFailed
 
 rem ## Success!
-goto Exit
+popd
+exit /B 0
 
 
 :Error_BatchFileInWrongLocation
@@ -110,4 +111,4 @@ goto Exit
 :Exit
 rem ## Restore original CWD in case we change it
 popd
-
+exit /B 1
\ No newline at end of file
