[INTEGRATE] Merging CL#2612583 using UE4-To-UE4-LauncherDev:
Let CreateDMG.sh fail on error instead of continuing blindly

diff --git a/Engine/Build/BatchFiles/Mac/CreateDMG.sh b/Engine/Build/BatchFiles/Mac/CreateDMG.sh
--- a/Engine/Build/BatchFiles/Mac/CreateDMG.sh
+++ b/Engine/Build/BatchFiles/Mac/CreateDMG.sh
@@ -1,4 +1,8 @@
-#/bin/sh
+#!/bin/bash
+
+# Automatically abort if any command returns a non-zero exit code.
+# Append something like "|| true" to the end of the command line if you really need to ignore errors for some reason.
+set -e
 
 APP_PATH=$(cd "$1" > /dev/null; pwd)
 APP_SIZE=$(expr $(du -sm "$APP_PATH" | awk '{print $1}') + 500)
@@ -82,4 +86,4 @@ test -f "$DMG_PATH" && rm -f "$DMG_PATH"
 hdiutil convert "$DMG_TEMP_PATH" -format UDZO -imagekey zlib-level=9 -o "$DMG_PATH" > /dev/null
 rm -f "$DMG_TEMP_PATH"
 
-exit $?
+exit 0
