Build.bat now properly sets the ERRORLEVEL upon exiting.

diff --git a/Engine/Build/BatchFiles/Build.bat b/Engine/Build/BatchFiles/Build.bat
--- a/Engine/Build/BatchFiles/Build.bat
+++ b/Engine/Build/BatchFiles/Build.bat
@@ -1,4 +1,5 @@
 @echo off
+setlocal enabledelayedexpansion
 
 REM The %~dp0 specifier resolves to the path to the directory where this .bat is located in.
 REM We use this so that regardless of where the .bat file was executed from, we can change to
@@ -12,6 +13,7 @@ REM %3 is the configuration name
 IF EXIST ..\..\Engine\Binaries\DotNET\UnrealBuildTool.exe (
          ..\..\Engine\Binaries\DotNET\UnrealBuildTool.exe %* -DEPLOY
 		 popd
+		 EXIT /B !RET!
 ) ELSE (
 	ECHO UnrealBuildTool.exe not found in ..\..\Engine\Binaries\DotNET\UnrealBuildTool.exe 
 	popd
