Linux: do not build LND as part of the setup.

- Not using LinuxNativeDialogs anymore.

diff --git a/Engine/Build/BatchFiles/Linux/Setup.sh b/Engine/Build/BatchFiles/Linux/Setup.sh
--- a/Engine/Build/BatchFiles/Linux/Setup.sh
+++ b/Engine/Build/BatchFiles/Linux/Setup.sh
@@ -161,9 +161,10 @@ CreateLinkIfNoneExists ../../Engine/shaders/Fxaa3_11.usf  ../Engine/Shaders/Fxaa
 echo Removing a stable libLND.so binary that was relocated in 4.8
 rm -f ../Engine/Binaries/Linux/libLND.so
 
-echo
-pushd Build/BatchFiles/Linux > /dev/null
-./BuildThirdParty.sh
-popd > /dev/null
+# As of 4.9 no third party library needs to be (re)built on any distro, so do not call BuildThirdParty.sh at all
+#echo
+#pushd Build/BatchFiles/Linux > /dev/null
+#./BuildThirdParty.sh
+#popd > /dev/null
 
 touch Build/OneTimeSetupPerformed
