Insert delays after mounting the DMG and before trying to write to it. Also remove code to suppress log output.

Merging using UE4-To-UE4-LauncherDev

diff --git a/Engine/Build/BatchFiles/Mac/CreateDMG.sh b/Engine/Build/BatchFiles/Mac/CreateDMG.sh
--- a/Engine/Build/BatchFiles/Mac/CreateDMG.sh
+++ b/Engine/Build/BatchFiles/Mac/CreateDMG.sh
@@ -19,7 +19,9 @@ VOLUME_ICON_PATH=$5
 test -d "$MOUNT_DIR" && hdiutil detach "$MOUNT_DIR" > /dev/null
 test -f "$DMG_TEMP_PATH" && rm -f "$DMG_TEMP_PATH"
 hdiutil create -srcfolder "$APP_PATH" -volname "$VOLUME_NAME" -fs HFS+ -fsargs "-c c=64,a=16,e=16" -format UDRW -size ${APP_SIZE}m "$DMG_TEMP_PATH" > /dev/null
+sleep 5
 hdiutil attach -readwrite -noverify -noautoopen "$DMG_TEMP_PATH" > /dev/null
+sleep 5
 
 # Create /Applications link
 ln -s /Applications "$MOUNT_DIR/Applications"
