Linux/Mac: fix packaging from the editor in the absence of UAT binary.

diff --git a/Engine/Build/BatchFiles/RunUAT.sh b/Engine/Build/BatchFiles/RunUAT.sh
--- a/Engine/Build/BatchFiles/RunUAT.sh
+++ b/Engine/Build/BatchFiles/RunUAT.sh
@@ -40,7 +40,7 @@ if [ ! -f Build/BatchFiles/RunUAT.sh ]; then
 fi
 
 # see if we have the no compile arg
-if echo "${Args[@]}" | grep -q -i "\-nocompile"; then
+if echo "${Args[@]}" | grep -q -w -i "\-nocompile"; then
 	UATCompileArg=
 else
 	UATCompileArg=-compile
