Linux: Moved the script that builds (parts of) cross-toolchain to main.

diff --git a/Engine/Build/BatchFiles/Linux/BuildCrossToolchain_x86_64.sh b/Engine/Build/BatchFiles/Linux/BuildCrossToolchain_x86_64.sh
new file mode 100755
--- /dev/null
+++ b/Engine/Build/BatchFiles/Linux/BuildCrossToolchain_x86_64.sh
@@ -0,0 +1,81 @@
+#!/bin/sh
+# by rcl^epic
+
+SelfName=`basename $0`
+LinuxBuildTemp=linux-host-build
+Win64BuildTemp=win64-host-build
+
+# prereq:
+# mercurial autoconf gperf bison flex libtool ncurses-dev
+
+
+#############
+# build recent ct-ng
+
+hg clone http://crosstool-ng.org/hg/crosstool-ng -r "1.19.0"
+cd crosstool-ng
+
+./bootstrap
+[ $? -ne 0 ] && echo "$SelfName: Error encountered, exiting at line $LINENO" && exit 1
+
+./configure --enable-local
+[ $? -ne 0 ] && echo "$SelfName: Error encountered, exiting at line $LINENO" && exit 1
+
+make
+[ $? -ne 0 ] && echo "$SelfName: Error encountered, exiting at line $LINENO" && exit 1
+
+cd ..
+
+#############
+# build Linux-hosted toolchain that targets mingw64
+
+rm -rf $LinuxBuildTemp
+mkdir $LinuxBuildTemp
+
+cp linux-host.config $LinuxBuildTemp/.config 
+[ $? -ne 0 ] && echo "$SelfName: Error encountered, exiting at line $LINENO" && exit 1
+
+cd $LinuxBuildTemp
+../crosstool-ng/ct-ng build
+[ $? -ne 0 ] && echo "$SelfName: Error encountered, exiting at line $LINENO" && exit 1
+
+cd ..
+
+
+#############
+# build Windows-hosted toolchain that targets Linux
+
+rm -rf $Win64BuildTemp
+mkdir $Win64BuildTemp
+
+cp win64-host.config $Win64BuildTemp/.config
+[ $? -ne 0 ] && echo "$SelfName: Error encountered, exiting at line $LINENO" && exit 1
+
+cd $Win64BuildTemp
+../crosstool-ng/ct-ng build
+[ $? -ne 0 ] && echo "$SelfName: Error encountered, exiting at line $LINENO" && exit 1
+
+cd ..
+
+#############
+# copy the GNU toolchain here
+
+mkdir CrossToolchain
+cp -r ~/x-tools/x86_64-unknown-linux-gnu/*/ CrossToolchain
+[ $? -ne 0 ] && echo "$SelfName: Error encountered, exiting at line $LINENO" && exit 1
+
+#############
+# rearrange files the way we need them
+
+cd CrossToolchain
+find . -type d -exec chmod +w {} \;
+cp -r -L x86_64-unknown-linux-gnu/include .
+cp -r -L x86_64-unknown-linux-gnu/sysroot/lib64 .
+mkdir -p usr
+cp -r -L x86_64-unknown-linux-gnu/sysroot/usr/lib64 usr
+cp -r -L x86_64-unknown-linux-gnu/sysroot/usr/include usr
+rm -rf x86_64-unknown-linux-gnu
+
+cd ..
+
+
