Linux cross-toolchain: re-arranged the files and added missing configs.

diff --git a/Engine/Build/BatchFiles/Linux/BuildCrossToolchain_x86_64.sh b/Engine/Build/BatchFiles/Linux/BuildCrossToolchain_x86_64.sh
deleted file mode 100755
diff --git a/Engine/Build/BatchFiles/Linux/Toolchain_x86-64/BuildCrossToolchain.sh b/Engine/Build/BatchFiles/Linux/Toolchain_x86-64/BuildCrossToolchain.sh
new file mode 100755
--- /dev/null
+++ b/Engine/Build/BatchFiles/Linux/Toolchain_x86-64/BuildCrossToolchain.sh
@@ -0,0 +1,81 @@
+#!/bin/sh
+# by rcl^epic
+
+SelfName=`basename $0`
+LinuxBuildTemp=linux-host-build
+Win64BuildTemp=win64-host-build
+
+# prereq:
+# mercurial autoconf gperf bison flex libtool ncurses-dev
+
+
+#############
+# build recent ct-ng
+
+hg clone http://crosstool-ng.org/hg/crosstool-ng -r "1.19.0"
+cd crosstool-ng
+
+./bootstrap
+[ $? -ne 0 ] && echo "$SelfName: Error encountered, exiting at line $LINENO" && exit 1
+
+./configure --enable-local
+[ $? -ne 0 ] && echo "$SelfName: Error encountered, exiting at line $LINENO" && exit 1
+
+make
+[ $? -ne 0 ] && echo "$SelfName: Error encountered, exiting at line $LINENO" && exit 1
+
+cd ..
+
+#############
+# build Linux-hosted toolchain that targets mingw64
+
+rm -rf $LinuxBuildTemp
+mkdir $LinuxBuildTemp
+
+cp linux-host.config $LinuxBuildTemp/.config 
+[ $? -ne 0 ] && echo "$SelfName: Error encountered, exiting at line $LINENO" && exit 1
+
+cd $LinuxBuildTemp
+../crosstool-ng/ct-ng build
+[ $? -ne 0 ] && echo "$SelfName: Error encountered, exiting at line $LINENO" && exit 1
+
+cd ..
+
+
+#############
+# build Windows-hosted toolchain that targets Linux
+
+rm -rf $Win64BuildTemp
+mkdir $Win64BuildTemp
+
+cp win64-host.config $Win64BuildTemp/.config
+[ $? -ne 0 ] && echo "$SelfName: Error encountered, exiting at line $LINENO" && exit 1
+
+cd $Win64BuildTemp
+../crosstool-ng/ct-ng build
+[ $? -ne 0 ] && echo "$SelfName: Error encountered, exiting at line $LINENO" && exit 1
+
+cd ..
+
+#############
+# copy the GNU toolchain here
+
+mkdir CrossToolchain
+cp -r ~/x-tools/x86_64-unknown-linux-gnu/*/ CrossToolchain
+[ $? -ne 0 ] && echo "$SelfName: Error encountered, exiting at line $LINENO" && exit 1
+
+#############
+# rearrange files the way we need them
+
+cd CrossToolchain
+find . -type d -exec chmod +w {} \;
+cp -r -L x86_64-unknown-linux-gnu/include .
+cp -r -L x86_64-unknown-linux-gnu/sysroot/lib64 .
+mkdir -p usr
+cp -r -L x86_64-unknown-linux-gnu/sysroot/usr/lib64 usr
+cp -r -L x86_64-unknown-linux-gnu/sysroot/usr/include usr
+rm -rf x86_64-unknown-linux-gnu
+
+cd ..
+
+
diff --git a/Engine/Build/BatchFiles/Linux/Toolchain_x86-64/README.md b/Engine/Build/BatchFiles/Linux/Toolchain_x86-64/README.md
new file mode 100644
--- /dev/null
+++ b/Engine/Build/BatchFiles/Linux/Toolchain_x86-64/README.md
@@ -0,0 +1,77 @@
+How to build the cross-toolchain.
+=================================
+
+Big picture view, steps are described in detail later:
+
+1.1 Build gcc-based cross-toolchain (we only need binutils/sysroot from it) on a Linux machine 
+     and copy them to a your Windows one.
+1.2 Build/download clang on a Windows machine and add it to above toolchain.
+
+How to use the cross-toolchain.
+=================================
+
+2.1 Set LINUX_ROOT environment variable to point to asbolute (Windows) path of the toolchain.
+2.2 Restart Visual Studio and regenerate projects with UBT.
+
+
+Building gcc-based cross-toolchain
+==================================
+
+Copy BuildCrossToolchain.sh together with linux-host.config and win64-host.config to
+a Linux machine running reasonably recent distro. Make sure that the following
+pre-requisites are installed (using Debian package names):
+
+ mercurial autoconf gperf bison flex libtool ncurses-dev
+
+and of course make/gcc.
+Run the script, which will clone current crosstool-ng and will use it to perform
+so called "canadian cross", i.e. will build Linux-hosted toolchain which targets
+mingw64 first, and then will build Windows-hosted toolchain that targets Linux.
+
+Note that the choice of kernel and binutils is pretty conservative and you are free
+to experiment (using "ct-ng menuconfig") if you need more recent ones. Not every
+combination is buildable though.
+
+If the script finishes without error, the new toolchain will be copied to subfolder
+called CrossToolchain of the current folder. Copy it to your Windows machine,
+e.g. to C:\CrossToolchain
+
+Building clang
+==============
+
+Note that you can grab a pre-built clang from here: http://llvm.org/releases/download.html, but
+we haven't actually tried that, building ourselves for various reasons.
+
+Detailed instructions how to build clang are given here: http://clang.llvm.org/get_started.html
+
+To recap, you will need: cmake, python and a supported compiler (recent Visual Studio is Ok).
+Grab stable release of llvm and clang from here: http://llvm.org/releases/download.html#3.6.0
+and unpack it to your local folder (unpack llvm first, then unpack clang into llvm/tools subdirectory).
+
+Use cmake to generate project files and then build Release (or MinSizeRel) x64 configuration of "install" 
+(or INSTALL, if you are using Visual Studio) target. Here's an example how to do the above from comamnd line:
+
+cmake -G "Visual Studio 11 Win64" -DPYTHON_EXECUTABLE=full_path_to_pythonw.exe -DCMAKE_INSTALL_PREFIX=path_to_copy_results_to ..\llvm
+devenv.exe LLVM.sln /Build Release /Project INSTALL
+
+Having done that, copy from install directory the following
+
+- bin/clang.exe
+- bin/clang++.exe
+- lib/clang/3.6.0/include 
+
+into the toolchain (so their relative path from toolchain root stays the same).
+
+
+Using toolchain to build UE4 Linux targets
+==========================================
+
+Set environment variable (Control Panel->System->Advanced system settings->Advanced->Environment variables)
+LINUX_ROOT to absolute path to your toolchain, without trailing backslash, so it looks like this:
+
+LINUX_ROOT=C:\CrossToolchain
+
+After you made sure that variable is set, you can restart Visual Studio and regenerate UE4 project files
+(using GenerateProjectFiles.bat). After this, you should have "Linux" available among Win32/Win64
+configurations, and you should be able to cross-compile for it. Note that for the moment being,
+only Server (Debug Server, Development Server, etc) targets are buildable for Linux.
