Only clean the UBT target if the list of .cs files in the source directory has changed. The project conditionally includes platform-specific files if they exist, but MSBuild doesn't check for new files when dependency checking. Previous fix was just to clean all the time, but comparing the file list is faster.

diff --git a/Engine/Build/BatchFiles/GenerateProjectFiles.bat b/Engine/Build/BatchFiles/GenerateProjectFiles.bat
--- a/Engine/Build/BatchFiles/GenerateProjectFiles.bat
+++ b/Engine/Build/BatchFiles/GenerateProjectFiles.bat
@@ -1,5 +1,5 @@
 @echo off
-
+setlocal
 echo Setting up Unreal Engine 4 project files...
 
 rem ## Unreal Engine 4 Visual Studio project setup script
@@ -55,7 +55,17 @@ goto Error_NoVisualStudioEnvironment
 
 
 :ReadyToCompile
-msbuild /nologo /verbosity:quiet Programs\UnrealBuildTool\UnrealBuildTool.csproj /property:Configuration=Development /property:Platform=AnyCPU /target:Clean,Build
+rem Check to see if the files in the UBT directory have changed. We conditionally include platform files from the .csproj file, but MSBuild doesn't recognize the dependency when new files are added. 
+md ..\Intermediate\Build >nul 2>nul
+dir /s /b Programs\UnrealBuildTool\*.cs >..\Intermediate\Build\UnrealBuildToolFiles.txt
+fc /b ..\Intermediate\Build\UnrealBuildToolFiles.txt ..\Intermediate\Build\UnrealBuildToolPrevFiles.txt >nul 2>nul
+if not errorlevel 1 goto SkipClean
+
+copy /y ..\Intermediate\Build\UnrealBuildToolFiles.txt ..\Intermediate\Build\UnrealBuildToolPrevFiles.txt >nul
+msbuild /nologo /verbosity:quiet Programs\UnrealBuildTool\UnrealBuildTool.csproj /property:Configuration=Development /property:Platform=AnyCPU /target:Clean
+
+:SkipClean
+msbuild /nologo /verbosity:quiet Programs\UnrealBuildTool\UnrealBuildTool.csproj /property:Configuration=Development /property:Platform=AnyCPU /target:Build
 if not %ERRORLEVEL% == 0 goto Error_UBTCompileFailed
 
 rem ## Run UnrealBuildTool to generate Visual Studio solution and project files
