Try to get the MSBuild path directly, without setting up the Visual Studio compile environment. Saves ~1-2 seconds per invocation.

diff --git a/Engine/Build/BatchFiles/GenerateProjectFiles.bat b/Engine/Build/BatchFiles/GenerateProjectFiles.bat
--- a/Engine/Build/BatchFiles/GenerateProjectFiles.bat
+++ b/Engine/Build/BatchFiles/GenerateProjectFiles.bat
@@ -27,6 +27,19 @@ rem ## Check to see if we're already running under a Visual Studio environment s
 if not "%INCLUDE%" == "" if not "%LIB%" == "" goto ReadyToCompile
 
 
+rem ## Try to get the MSBuild executable directly (see https://msdn.microsoft.com/en-us/library/hh162058(v=vs.120).aspx)
+
+set MSBUILD_EXE=msbuild.exe
+
+if exist "%ProgramFiles%\MSBuild\12.0\bin\MSBuild.exe" (
+ 	set MSBUILD_EXE="%ProgramFiles%\MSBuild\12.0\bin\MSBuild.exe"
+	goto ReadyToCompile
+)
+if exist "%ProgramFiles(x86)%\MSBuild\12.0\bin\MSBuild.exe" (
+	set MSBUILD_EXE="%ProgramFiles(x86)%\MSBuild\12.0\bin\MSBuild.exe"
+	goto ReadyToCompile
+)
+
 rem ## Check for Visual Studio 2013
 
 pushd %~dp0
@@ -62,10 +75,10 @@ fc /b ..\Intermediate\Build\UnrealBuildToolFiles.txt ..\Intermediate\Build\Unrea
 if not errorlevel 1 goto SkipClean
 
 copy /y ..\Intermediate\Build\UnrealBuildToolFiles.txt ..\Intermediate\Build\UnrealBuildToolPrevFiles.txt >nul
-msbuild /nologo /verbosity:quiet Programs\UnrealBuildTool\UnrealBuildTool.csproj /property:Configuration=Development /property:Platform=AnyCPU /target:Clean
+%MSBUILD_EXE% /nologo /verbosity:quiet Programs\UnrealBuildTool\UnrealBuildTool.csproj /property:Configuration=Development /property:Platform=AnyCPU /target:Clean
 
 :SkipClean
-msbuild /nologo /verbosity:quiet Programs\UnrealBuildTool\UnrealBuildTool.csproj /property:Configuration=Development /property:Platform=AnyCPU /target:Build
+%MSBUILD_EXE% /nologo /verbosity:quiet Programs\UnrealBuildTool\UnrealBuildTool.csproj /property:Configuration=Development /property:Platform=AnyCPU /target:Build
 if not %ERRORLEVEL% == 0 goto Error_UBTCompileFailed
 
 rem ## Run UnrealBuildTool to generate Visual Studio solution and project files
