Bump up Mono target to 4.5 due to IReadOnlyList usage.

diff --git a/Engine/Build/BatchFiles/Linux/GenerateProjectFiles.sh b/Engine/Build/BatchFiles/Linux/GenerateProjectFiles.sh
--- a/Engine/Build/BatchFiles/Linux/GenerateProjectFiles.sh
+++ b/Engine/Build/BatchFiles/Linux/GenerateProjectFiles.sh
@@ -21,7 +21,7 @@ echo
 set -x
 xbuild Source/Programs/UnrealBuildTool/UnrealBuildTool.csproj \
   /verbosity:quiet /nologo \
-  /p:TargetFrameworkVersion=v4.0 \
+  /p:TargetFrameworkVersion=v4.5 \
   /p:Configuration="Development"
 
 # pass all parameters to UBT
