Propagate the correct exit code from UAT out of RunUAT.bat

diff --git a/Engine/Build/BatchFiles/RunUAT.bat b/Engine/Build/BatchFiles/RunUAT.bat
--- a/Engine/Build/BatchFiles/RunUAT.bat
+++ b/Engine/Build/BatchFiles/RunUAT.bat
@@ -83,22 +83,27 @@ goto Exit
 
 :Error_BatchFileInWrongLocation
 echo RunUAT.bat ERROR: The batch file does not appear to be located in the /Engine/Build/BatchFiles directory.  This script must be run from within that directory.
+set RUNUAT_EXITCODE=1
 goto Exit_Failure
 
 :Error_NoVisualStudioEnvironment
 echo RunUAT.bat ERROR: A valid version of Visual Studio 2013 or Visual Studio 2012 does not appear to be installed.
+set RUNUAT_EXITCODE=1
 goto Exit_Failure
 
 :Error_NoFallbackExecutable
 echo RunUAT.bat ERROR: Visual studio and/or AutomationTool.csproj was not found, nor was Engine\Binaries\DotNET\AutomationTool.exe. Can't run the automation tool.
+set RUNUAT_EXITCODE=1
 goto Exit_Failure
 
 :Error_UATCompileFailed
 echo RunUAT.bat ERROR: AutomationTool failed to compile.
+set RUNUAT_EXITCODE=1
 goto Exit_Failure
 
 
 :Error_UATFailed
+set RUNUAT_EXITCODE=%ERRORLEVEL%
 echo copying UAT log files...
 if not "%uebp_LogFolder%" == "" copy log*.txt %uebp_LogFolder%\UAT_*.*
 rem if "%uebp_LogFolder%" == "" copy log*.txt c:\LocalBuildLogs\UAT_*.*
@@ -109,7 +114,7 @@ goto Exit_Failure
 :Exit_Failure
 echo BUILD FAILED
 popd
-exit /B %ERRORLEVEL%
+exit /B %RUNUAT_EXITCODE%
 
 :Exit
 rem ## Restore original CWD in case we change it
