UBT: Remove unused compiler arguments in Linux builds.

- Also fix a shell script for incremental rebuilds.

diff --git a/Engine/Build/BatchFiles/Linux/Build.sh b/Engine/Build/BatchFiles/Linux/Build.sh
--- a/Engine/Build/BatchFiles/Linux/Build.sh
+++ b/Engine/Build/BatchFiles/Linux/Build.sh
@@ -7,7 +7,7 @@ set -e
 cd "`dirname "$0"`/../../../.." 
 
 # First make sure that the UnrealBuildTool is up-to-date
-if ! xbuild /property:Configuration=Development /property:TargetFrameworkVersion=v4.0 /verbosity:quiet /nologo Engine/Source/Programs/UnrealBuildTool/UnrealBuildTool.csproj; then
+if ! xbuild /property:Configuration=Development /property:TargetFrameworkVersion=v4.5 /verbosity:quiet /nologo Engine/Source/Programs/UnrealBuildTool/UnrealBuildTool.csproj; then
   echo "Failed to build to build tool (UnrealBuildTool)"
   exit 1
 fi
diff --git a/Engine/Source/Programs/UnrealBuildTool/Linux/LinuxToolChain.cs b/Engine/Source/Programs/UnrealBuildTool/Linux/LinuxToolChain.cs
--- a/Engine/Source/Programs/UnrealBuildTool/Linux/LinuxToolChain.cs
+++ b/Engine/Source/Programs/UnrealBuildTool/Linux/LinuxToolChain.cs
@@ -215,11 +215,7 @@ namespace UnrealBuildTool
 		{
 			string Result = "";
 
-			if (Architecture.StartsWith("x86_64"))
-			{
-				Result += " -mmmx -msse -msse2";
-			}
-			else if (Architecture.StartsWith("arm"))
+			if (Architecture.StartsWith("arm"))
 			{
 				Result += " -fsigned-char";
 
