Remove the "ERROR: Automation Tool was unable to run succesfully" message. It's redundant, and gets picked up unnecessarily by the EC postp.

diff --git a/Engine/Build/BatchFiles/RunUAT.bat b/Engine/Build/BatchFiles/RunUAT.bat
--- a/Engine/Build/BatchFiles/RunUAT.bat
+++ b/Engine/Build/BatchFiles/RunUAT.bat
@@ -108,7 +108,6 @@ echo copying UAT log files...
 if not "%uebp_LogFolder%" == "" copy log*.txt %uebp_LogFolder%\UAT_*.*
 rem if "%uebp_LogFolder%" == "" copy log*.txt c:\LocalBuildLogs\UAT_*.*
 popd
-echo RunUAT.bat ERROR: AutomationTool was unable to run successfully.
 goto Exit_Failure
 
 :Exit_Failure
