PR #1577: add tanglu to install script (Contributed by [REDACTED_PERSON])

diff --git a/Engine/Build/BatchFiles/Linux/Setup.sh b/Engine/Build/BatchFiles/Linux/Setup.sh
--- a/Engine/Build/BatchFiles/Linux/Setup.sh
+++ b/Engine/Build/BatchFiles/Linux/Setup.sh
@@ -32,7 +32,7 @@ fi
 if [ -e /etc/os-release ]; then
   source /etc/os-release
   # Ubuntu/Debian/Mint
-  if [[ "$ID" == "ubuntu" ]] || [[ "$ID_LIKE" == "ubuntu" ]] || [[ "$ID" == "debian" ]] || [[ "$ID_LIKE" == "debian" ]]; then
+  if [[ "$ID" == "ubuntu" ]] || [[ "$ID_LIKE" == "ubuntu" ]] || [[ "$ID" == "debian" ]] || [[ "$ID_LIKE" == "debian" ]] || [[ "$ID" == "tanglu" ]] || [[ "$ID_LIKE" == "tanglu" ]]; then
     # Install the necessary dependencies (require clang-3.5, although 3.3 and 3.6 should work too for this release)
     DEPS="mono-xbuild \
       mono-dmcs \
