Support for p4 tool installed in /usr/local/bin on Mac

diff --git a/Engine/Build/BatchFiles/Mac/xcodeunlock.sh b/Engine/Build/BatchFiles/Mac/xcodeunlock.sh
--- a/Engine/Build/BatchFiles/Mac/xcodeunlock.sh
+++ b/Engine/Build/BatchFiles/Mac/xcodeunlock.sh
@@ -29,6 +29,14 @@ fb=0
 # Xcode doesn't setup the environment, so we load it here:
 source ~/.bash_profile
 
+export PATH="/usr/local/bin:/usr/bin:$PATH"
+
+P4TOOL_PATH=`which p4`
+
+if [ $P4TOOL_PATH == "" ] || [ ! -f $P4TOOL_PATH ]; then
+	P4TOOL_PATH = "/usr/local/bin/p4"
+fi
+
 fn=${XcodeAlertAffectedPaths}
 fn=$(printf $(echo -n $fn | sed 's/\\/\\\\/g;s/\(%\)\([0-9a-fA-F][0-9a-fA-F]\)/\\x\2/g') )
 
@@ -37,7 +45,7 @@ fp=$(dirname "${fn}")
 echo "fp=" ${fp}
 
 # A check to confirm the current P4CONFIG setting:
-conf=$(/usr/bin/p4 -d "${fp}" set P4CONFIG 2>&1)
+conf=$($P4TOOL_PATH -d "${fp}" set P4CONFIG 2>&1)
 echo ${conf}
 
 if [ -a "${fn}" ]; then
@@ -46,7 +54,7 @@ if [ -a "${fn}" ]; then
     export XC_TO=${to}
 
     # Check the file out from Perforce:
-    res=$(/usr/bin/p4 -d "${fp}" edit "${fn}" 2>&1)
+    res=$($P4TOOL_PATH -d "${fp}" edit "${fn}" 2>&1)
     echo "res=" ${res}
 
     # Save the result as an environment variable
diff --git a/Engine/Source/Programs/AutomationTool/AutomationUtils/MacHostPlatform.cs b/Engine/Source/Programs/AutomationTool/AutomationUtils/MacHostPlatform.cs
--- a/Engine/Source/Programs/AutomationTool/AutomationUtils/MacHostPlatform.cs
+++ b/Engine/Source/Programs/AutomationTool/AutomationUtils/MacHostPlatform.cs
@@ -65,9 +65,22 @@ namespace AutomationTool
 			get { return Path.Combine(Environment.GetFolderPath(Environment.SpecialFolder.Personal), "Library/Logs/Unreal Engine/LocalBuildLogs"); }
 		}
 
+		private string P4ExePath = null;
+
 		public override string P4Exe
 		{
-			get { return "/usr/bin/p4"; }
+			get
+			{
+				if (P4ExePath == null)
+				{
+					P4ExePath = "/usr/bin/p4";
+					if (!File.Exists(P4ExePath))
+					{
+						P4ExePath = "/usr/local/bin/p4";
+					}
+				}
+				return P4ExePath;
+			}
 		}
 
 		public override Process CreateProcess(string AppName)
