Linux: install more mono packages during setup (UE-22018).

- Explicitly install System.IO.Compression.[Filesystem] packages which are required for UAT compilation (they might have been already installed through dependencies, but this is more robust).

diff --git a/Engine/Build/BatchFiles/Linux/Setup.sh b/Engine/Build/BatchFiles/Linux/Setup.sh
--- a/Engine/Build/BatchFiles/Linux/Setup.sh
+++ b/Engine/Build/BatchFiles/Linux/Setup.sh
@@ -43,6 +43,8 @@ if [ -e /etc/os-release ]; then
       libmono-system-xml-linq4.0-cil
       libmono-corlib4.0-cil
       libmono-windowsbase4.0-cil
+      libmono-system-io-compression4.0-cil
+      libmono-system-io-compression-filesystem4.0-cil
       clang-3.5
       "
 
