Increase timeouts in CreateDMG.sh to improve stability.
Addresses intermittent build failures which are affecting builds in //UE4/Main.

diff --git a/Engine/Build/BatchFiles/Mac/CreateDMG.sh b/Engine/Build/BatchFiles/Mac/CreateDMG.sh
--- a/Engine/Build/BatchFiles/Mac/CreateDMG.sh
+++ b/Engine/Build/BatchFiles/Mac/CreateDMG.sh
@@ -19,9 +19,9 @@ VOLUME_ICON_PATH=$5
 test -d "$MOUNT_DIR" && hdiutil detach "$MOUNT_DIR" > /dev/null
 test -f "$DMG_TEMP_PATH" && rm -f "$DMG_TEMP_PATH"
 hdiutil create -srcfolder "$APP_PATH" -volname "$VOLUME_NAME" -fs HFS+ -fsargs "-c c=64,a=16,e=16" -format UDRW -size ${APP_SIZE}m "$DMG_TEMP_PATH" > /dev/null
-sleep 5
+sleep 8
 hdiutil attach -readwrite -noverify -noautoopen "$DMG_TEMP_PATH" > /dev/null
-sleep 5
+sleep 8
 
 # Create /Applications link
 ln -s /Applications "$MOUNT_DIR/Applications"
@@ -41,7 +41,7 @@ echo '
 		tell disk "'$VOLUME_NAME'"
 			open
 
-			delay 8
+			delay 12
 
 			tell container window
 				set current view to icon view
@@ -50,7 +50,7 @@ echo '
 				set the bounds to {200, 200, 840, 575}
 			end tell
 
-			delay 8
+			delay 12
 
 			set Options to the icon view options of container window
 			tell Options
@@ -59,16 +59,16 @@ echo '
 			end tell
 			set background picture of Options to file ".background:'$BACKGROUND_FILE_NAME'"
 
-			delay 8
+			delay 12
 
 			set position of item "'$APP_NAME'" of container window to {170, 220}
 			set position of item "Applications" of container window to {485, 220}
 
-			delay 8
+			delay 12
 
 			update without registering applications
 
-			delay 8
+			delay 12
 			close
 		end tell
 	end tell
