If RunUAT gets a signal while running mono bring mono down as well

diff --git a/Engine/Build/BatchFiles/RunUAT.sh b/Engine/Build/BatchFiles/RunUAT.sh
--- a/Engine/Build/BatchFiles/RunUAT.sh
+++ b/Engine/Build/BatchFiles/RunUAT.sh
@@ -10,6 +10,10 @@ echo
 echo Running AutomationTool...
 echo
 
+TermHandler() {
+	kill -TERM $UATPid 2>&1 /dev/null
+}
+
 # loop over the arguments, quoting spaces to pass to UAT proper
 Args=
 i=0
@@ -91,10 +95,15 @@ if [ -z "$uebp_LogFolder" ]; then
 else
 	LogDir="$uebp_LogFolder"
 fi
+
+trap TermHandler SIGTERM SIGINT
+
 # you can't set a dotted env var nicely in sh, but env will run a command with
 # a list of env vars set, including dotted ones
 echo Start UAT: mono AutomationTool.exe "${Args[@]}"
 env uebp_LogFolder="$LogDir" mono AutomationTool.exe "${Args[@]}" $UATCompileArg
+UATPid=$!
+wait $UATPid
 UATReturn=$?
 
 # @todo: Copy log files to somewhere useful
