Set our location to Linux/Mac GenerateProject.sh if we are not running from that location

diff --git a/Engine/Build/BatchFiles/Linux/GenerateProjectFiles.sh b/Engine/Build/BatchFiles/Linux/GenerateProjectFiles.sh
--- a/Engine/Build/BatchFiles/Linux/GenerateProjectFiles.sh
+++ b/Engine/Build/BatchFiles/Linux/GenerateProjectFiles.sh
@@ -5,23 +5,26 @@ echo
 echo Setting up Unreal Engine 4 project files...
 echo
 
+# If ran from someone other then the script location we'll have the full base path
+BASE_PATH="`dirname "$0"`"
+
 # this is located inside an extra 'Linux' path unlike the Windows variant.
 
-if [ ! -d ../../../Binaries/DotNET ]; then
+if [ ! -d "$BASE_PATH/../../../Binaries/DotNET" ]; then
  echo GenerateProjectFiles ERROR: It looks like you're missing some files that are required in order to generate projects.  Please check that you've downloaded and unpacked the engine source code, binaries, content and third-party dependencies before running this script.
  exit 1
 fi
 
-if [ ! -d ../../../Source ]; then
+if [ ! -d "$BASE_PATH/../../../Source" ]; then
  echo GenerateProjectFiles ERROR: This script file does not appear to be located inside the Engine/Build/BatchFiles/Mac directory.
  exit 1
 fi
 
-source SetupMono.sh "`dirname "$0"`"
+source "$BASE_PATH/SetupMono.sh" $BASE_PATH
 
-if [ -f ../../../Source/Programs/UnrealBuildTool/UnrealBuildTool.csproj ]; then
-	xbuild ../../../Source/Programs/UnrealBuildTool/UnrealBuildTool.csproj /property:Configuration="Development" /p:TargetFrameworkVersion=v4.5 /verbosity:quiet /nologo
+if [ -f "$BASE_PATH/../../../Source/Programs/UnrealBuildTool/UnrealBuildTool.csproj" ]; then
+	xbuild "$BASE_PATH/../../../Source/Programs/UnrealBuildTool/UnrealBuildTool.csproj" /property:Configuration="Development" /p:TargetFrameworkVersion=v4.5 /verbosity:quiet /nologo
 fi
 
 # pass all parameters to UBT
-mono ../../../Binaries/DotNET/UnrealBuildTool.exe -projectfiles "$@"
+mono "$BASE_PATH/../../../Binaries/DotNET/UnrealBuildTool.exe" -projectfiles "$@"
diff --git a/Engine/Build/BatchFiles/Mac/GenerateProjectFiles.sh b/Engine/Build/BatchFiles/Mac/GenerateProjectFiles.sh
--- a/Engine/Build/BatchFiles/Mac/GenerateProjectFiles.sh
+++ b/Engine/Build/BatchFiles/Mac/GenerateProjectFiles.sh
@@ -4,23 +4,26 @@ echo
 echo Setting up Unreal Engine 4 project files...
 echo
 
+# If ran from somewhere other then the script location we'll have the full base path
+BASE_PATH="`dirname "$0"`"
+
 # this is located inside an extra 'Mac' path unlike the Windows variant.
 
-if [ ! -d ../../../Binaries/DotNET ]; then
+if [ ! -d "$BASE_PATH/../../../Binaries/DotNET" ]; then
  echo GenerateProjectFiles ERROR: It looks like you're missing some files that are required in order to generate projects.  Please check that you've downloaded and unpacked the engine source code, binaries, content and third-party dependencies before running this script.
  exit 1
 fi
 
-if [ ! -d ../../../Source ]; then
+if [ ! -d "$BASE_PATH/../../../Source" ]; then
  echo GenerateProjectFiles ERROR: This script file does not appear to be located inside the Engine/Build/BatchFiles/Mac directory.
  exit 1
 fi
 
-source SetupMono.sh "`dirname "$0"`"
+. "$BASE_PATH/SetupMono.sh" $BASE_PATH
 
-if [ -f ../../../Source/Programs/UnrealBuildTool/UnrealBuildTool.csproj ]; then
-	xbuild ../../../Source/Programs/UnrealBuildTool/UnrealBuildTool.csproj /property:Configuration="Development" /verbosity:quiet /nologo |grep -i error
+if [ -f "$BASE_PATH/../../../Source/Programs/UnrealBuildTool/UnrealBuildTool.csproj" ]; then
+	xbuild "$BASE_PATH/../../../Source/Programs/UnrealBuildTool/UnrealBuildTool.csproj" /property:Configuration="Development" /verbosity:quiet /nologo |grep -i error
 fi
 
 # pass all parameters to UBT
-mono ../../../Binaries/DotNET/UnrealBuildTool.exe -projectfiles "$@"
+mono "$BASE_PATH/../../../Binaries/DotNET/UnrealBuildTool.exe" -projectfiles "$@"
