Add an error when generating project files if the NuGet package manager is not installed.

diff --git a/Engine/Build/BatchFiles/GenerateProjectFiles.bat b/Engine/Build/BatchFiles/GenerateProjectFiles.bat
--- a/Engine/Build/BatchFiles/GenerateProjectFiles.bat
+++ b/Engine/Build/BatchFiles/GenerateProjectFiles.bat
@@ -29,6 +29,22 @@ rem ## Get the path to MSBuild
 call "%~dp0GetMSBuildPath.bat"
 if errorlevel 1 goto Error_NoVisualStudioEnvironment
 
+rem ## If we're using VS2017, check that NuGet package manager is installed. MSBuild fails to compile C# projects from the command line with a cryptic error if it's not: 
+rem ## https://developercommunity.visualstudio.com/content/problem/137779/the-getreferencenearesttargetframeworktask-task-wa.html
+if not exist "%ProgramFiles(x86)%\Microsoft Visual Studio\Installer\vswhere.exe" goto NoVsWhere
+
+set MSBUILD_EXE_WITH_NUGET=
+for /f "delims=" %%i in ('"%ProgramFiles(x86)%\Microsoft Visual Studio\Installer\vswhere" -latest -products * -requires Microsoft.Component.MSBuild -requires Microsoft.VisualStudio.Component.NuGet -property installationPath') do (
+	if exist "%%i\MSBuild\15.0\Bin\MSBuild.exe" (
+		set MSBUILD_EXE_WITH_NUGET="%%i\MSBuild\15.0\Bin\MSBuild.exe"
+		goto FoundMsBuildWithNuget
+	)
+)
+
+:FoundMsBuildWithNuget
+if not [%MSBUILD_EXE_WITH_NUGET%] == [%MSBUILD_EXE%] goto Error_RequireNugetPackageManager
+
+:NoVsWhere
 
 rem Check to see if the files in the UBT directory have changed. We conditionally include platform files from the .csproj file, but MSBuild doesn't recognize the dependency when new files are added. 
 md ..\Intermediate\Build >nul 2>nul
@@ -74,6 +90,12 @@ echo.
 pause
 goto Exit
 
+:Error_RequireNugetPackageManager
+echo.
+echo UE4 requires the NuGet Package Manager to be installed. Please run the Visual Studio Installer and add it from the individual components list (in the 'Code Tools' category).
+echo.
+pause
+goto Exit
 
 :Error_UBTCompileFailed
 echo.
