Fix GenerateProjectFiles.bat failing with error about requiring NuGet if VS2017 is not installed.

diff --git a/Engine/Build/BatchFiles/GenerateProjectFiles.bat b/Engine/Build/BatchFiles/GenerateProjectFiles.bat
--- a/Engine/Build/BatchFiles/GenerateProjectFiles.bat
+++ b/Engine/Build/BatchFiles/GenerateProjectFiles.bat
@@ -33,16 +33,26 @@ rem ## If we're using VS2017, check that NuGet package manager is installed. MSB
 rem ## https://developercommunity.visualstudio.com/content/problem/137779/the-getreferencenearesttargetframeworktask-task-wa.html
 if not exist "%ProgramFiles(x86)%\Microsoft Visual Studio\Installer\vswhere.exe" goto NoVsWhere
 
-set MSBUILD_EXE_WITH_NUGET=
+set MSBUILD_15_EXE=
+for /f "delims=" %%i in ('"%ProgramFiles(x86)%\Microsoft Visual Studio\Installer\vswhere" -latest -products * -requires Microsoft.Component.MSBuild -property installationPath') do (
+	if exist "%%i\MSBuild\15.0\Bin\MSBuild.exe" (
+		set MSBUILD_15_EXE="%%i\MSBuild\15.0\Bin\MSBuild.exe"
+		goto FoundMsBuild15
+	)
+)
+:FoundMsBuild15
+
+set MSBUILD_15_EXE_WITH_NUGET=
 for /f "delims=" %%i in ('"%ProgramFiles(x86)%\Microsoft Visual Studio\Installer\vswhere" -latest -products * -requires Microsoft.Component.MSBuild -requires Microsoft.VisualStudio.Component.NuGet -property installationPath') do (
 	if exist "%%i\MSBuild\15.0\Bin\MSBuild.exe" (
-		set MSBUILD_EXE_WITH_NUGET="%%i\MSBuild\15.0\Bin\MSBuild.exe"
-		goto FoundMsBuildWithNuget
+		set MSBUILD_15_EXE_WITH_NUGET="%%i\MSBuild\15.0\Bin\MSBuild.exe"
+		goto FoundMsBuild15WithNuget
 	)
 )
+:FoundMsBuild15WithNuget
 
-:FoundMsBuildWithNuget
-if not [%MSBUILD_EXE_WITH_NUGET%] == [%MSBUILD_EXE%] goto Error_RequireNugetPackageManager
+if not [%MSBUILD_EXE%] == [%MSBUILD_15_EXE%] goto NoVsWhere
+if not [%MSBUILD_EXE%] == [%MSBUILD_15_EXE_WITH_NUGET%] goto Error_RequireNugetPackageManager
 
 :NoVsWhere
 
@@ -92,7 +102,7 @@ goto Exit
 
 :Error_RequireNugetPackageManager
 echo.
-echo UE4 requires the NuGet Package Manager to be installed. Please run the Visual Studio Installer and add it from the individual components list (in the 'Code Tools' category).
+echo UE4 requires the NuGet Package Manager to be installed to use %MSBUILD_EXE%. Please run the Visual Studio Installer and add it from the individual components list (in the 'Code Tools' category).
 echo.
 pause
 goto Exit
