Dont install unsupported version of clang and SetupLinuxToolchain instead

diff --git a/Engine/Build/BatchFiles/Linux/Setup.sh b/Engine/Build/BatchFiles/Linux/Setup.sh
--- a/Engine/Build/BatchFiles/Linux/Setup.sh
+++ b/Engine/Build/BatchFiles/Linux/Setup.sh
@@ -18,6 +18,8 @@ set -e
 TOP_DIR=$(cd "$SCRIPT_DIR/../../.." ; pwd)
 cd "${TOP_DIR}"
 
+NEED_TOOLCHAIN=false
+
 if [ -e /etc/os-release ]; then
   source /etc/os-release
   # Ubuntu/Debian/Mint
@@ -47,9 +49,9 @@ if [ -e /etc/os-release ]; then
        libmono-system-io-compression4.0-cil
        libmono-system-io-compression-filesystem4.0-cil
        mono-devel
-       clang-3.5
        build-essential
        "
+      NEED_TOOLCHAIN=true
     elif [ -n "$VERSION_ID" ] && [[ "$VERSION_ID" == 16.04 ]]; then
      DEPS="mono-xbuild \
        mono-dmcs \
@@ -202,6 +204,15 @@ if [ -e /etc/os-release ]; then
   fi
 fi
 
+# If we are using a distro that doesn't have a supported version of clang
+# we need to setup our toolchain
+if $NEED_TOOLCHAIN; then
+  echo "No supported version of clang found, setting up bundled toolchain"
+  pushd Build/BatchFiles/Linux > /dev/null
+  ./SetupToolchain.sh
+  popd > /dev/null
+fi
+
 # Provide the hooks for locally building third party libs if needed
 echo
 pushd Build/BatchFiles/Linux > /dev/null
diff --git a/Engine/Source/Programs/UnrealBuildTool/Platform/Linux/LinuxToolChain.cs b/Engine/Source/Programs/UnrealBuildTool/Platform/Linux/LinuxToolChain.cs
--- a/Engine/Source/Programs/UnrealBuildTool/Platform/Linux/LinuxToolChain.cs
+++ b/Engine/Source/Programs/UnrealBuildTool/Platform/Linux/LinuxToolChain.cs
@@ -134,7 +134,7 @@ namespace UnrealBuildTool
 			// disable that only if you are a dev and you know what you are doing
 			if (!UsingClang())
 			{
-				throw new BuildException("This version of the engine can only be compiled by clang - refusing to register the Linux toolchain.");
+				throw new BuildException("Unable to build: no compatible clang version found. Please run Setup.sh");
 			}
 			// prevent unknown clangs since the build is likely to fail on too old or too new compilers
 			else if ((CompilerVersionMajor * 10 + CompilerVersionMinor) > 60 || (CompilerVersionMajor * 10 + CompilerVersionMinor) < 38)
