By default install the bundled toolchain unless we are a perforce build

diff --git a/Engine/Build/BatchFiles/Linux/Setup.sh b/Engine/Build/BatchFiles/Linux/Setup.sh
--- a/Engine/Build/BatchFiles/Linux/Setup.sh
+++ b/Engine/Build/BatchFiles/Linux/Setup.sh
@@ -18,8 +18,6 @@ set -e
 TOP_DIR=$(cd "$SCRIPT_DIR/../../.." ; pwd)
 cd "${TOP_DIR}"
 
-NEED_TOOLCHAIN=false
-
 if [ -e /etc/os-release ]; then
   source /etc/os-release
   # Ubuntu/Debian/Mint
@@ -51,7 +49,6 @@ if [ -e /etc/os-release ]; then
        mono-devel
        build-essential
        "
-      NEED_TOOLCHAIN=true
     elif [ -n "$VERSION_ID" ] && [[ "$VERSION_ID" == 16.04 ]]; then
      DEPS="mono-xbuild \
        mono-dmcs \
@@ -66,8 +63,6 @@ if [ -e /etc/os-release ]; then
        libmono-system-io-compression-filesystem4.0-cil
        libmono-system-runtime4.0-cil
        mono-devel
-       clang-3.8
-       llvm
        build-essential
        "
     elif [ -n "$VERSION_ID" ] && [[ "$VERSION_ID" < 17.10 ]]; then
@@ -84,8 +79,6 @@ if [ -e /etc/os-release ]; then
        libmono-system-io-compression-filesystem4.0-cil
        libmono-system-runtime4.0-cil
        mono-devel
-       clang-3.9
-       llvm
        build-essential
        "
     elif [[ $PRETTY_NAME == *sid ]] || [[ $PRETTY_NAME == *stretch ]]; then
@@ -102,7 +95,6 @@ if [ -e /etc/os-release ]; then
        libmono-system-io-compression-filesystem4.0-cil
        libmono-system-runtime4.0-cil
        mono-devel
-       clang-3.8
        "
     else # assume the latest Ubuntu, this is going to be a moving target (17.10 as of now)
      DEPS="mono-xbuild \
@@ -118,20 +110,10 @@ if [ -e /etc/os-release ]; then
        libmono-system-io-compression-filesystem4.0-cil
        libmono-system-runtime4.0-cil
        mono-devel
-       clang-5.0
-       lld-5.0
-       llvm
        build-essential
        "
     fi
 
-    # these tools are only needed to build third-party software which is prebuilt for Ubuntu.
-    if [[ "$ID" != "ubuntu" ]]; then
-      DEPS+="libqt4-dev \
-             cmake
-            "
-    fi
-
     for DEP in $DEPS; do
       if ! PackageIsInstalled $DEP; then
         echo "Attempting installation of missing package: $DEP"
@@ -141,73 +123,11 @@ if [ -e /etc/os-release ]; then
       fi
     done
   fi
-  
-  # openSUSE/SLED/SLES
-  if [[ "$ID" == "opensuse" ]] || [[ "$ID_LIKE" == "suse" ]]; then
-    # Install all necessary dependencies
-    DEPS="mono-core
-      mono-devel
-      mono-mvc
-      libqt4-devel
-      dos2unix
-      cmake
-      "
-
-    for DEP in $DEPS; do
-      if ! rpm -q $DEP > /dev/null 2>&1; then
-        echo "Attempting installation of missing package: $DEP"
-        set -x
-        sudo zypper -n install $DEP
-        set +x
-      fi
-    done
-  fi
-
-  # Fedora
-  if [[ "$ID" == "fedora" ]]; then
-    # Install all necessary dependencies
-    DEPS="mono-core
-      mono-devel
-      qt-devel
-      dos2unix
-      cmake
-      clang
-      "
-
-    for DEP in $DEPS; do
-      if ! rpm -q $DEP > /dev/null 2>&1; then
-        echo "Attempting installation of missing package: $DEP"
-        set -x
-        sudo dnf -y install $DEP
-        set +x
-      fi
-    done
-  fi
-
-
-  # Arch Linux
-  if [[ "$ID" == "arch" ]] || [[ "$ID_LIKE" == "arch" ]]; then
-    DEPS="clang mono python sdl2 qt4 dos2unix cmake"
-    MISSING=false
-    for DEP in $DEPS; do
-      if ! pacman -Qs $DEP > /dev/null 2>&1; then
-        MISSING=true
-        break
-      fi
-    done
-    if [ "$MISSING" = true ]; then
-      echo "Attempting to install missing packages: $DEPS"
-      set -x
-      sudo pacman -S --needed --noconfirm $DEPS
-      set +x
-    fi
-  fi
 fi
 
-# If we are using a distro that doesn't have a supported version of clang
-# we need to setup our toolchain
-if $NEED_TOOLCHAIN; then
-  echo "No supported version of clang found, setting up bundled toolchain"
+# Install our bundled toolchain unless we are running a Perforce build
+if [ ! -f Build/PerforceBuild.txt ]; then
+  echo "Installing a bundled clang toolchain"
   pushd Build/BatchFiles/Linux > /dev/null
   ./SetupToolchain.sh
   popd > /dev/null
