Break if 0 exit code otherwise continue

diff --git a/Engine/Build/BatchFiles/Linux/GitDependencies.sh b/Engine/Build/BatchFiles/Linux/GitDependencies.sh
--- a/Engine/Build/BatchFiles/Linux/GitDependencies.sh
+++ b/Engine/Build/BatchFiles/Linux/GitDependencies.sh
@@ -24,7 +24,7 @@ RESULT=0
 source "$BASE_PATH/SetupMono.sh" $BASE_PATH
 
 while : ; do
-        mono Engine/Binaries/DotNET/GitDependencies.exe $ARGS
+        mono Engine/Binaries/DotNET/GitDependencies.exe $ARGS && break
         RESULT=$?
 
         echo "Result: $RESULT"
