Fix GitDeps exiting early if a command fails

diff --git a/Engine/Build/BatchFiles/Linux/GitDependencies.sh b/Engine/Build/BatchFiles/Linux/GitDependencies.sh
--- a/Engine/Build/BatchFiles/Linux/GitDependencies.sh
+++ b/Engine/Build/BatchFiles/Linux/GitDependencies.sh
@@ -6,8 +6,6 @@
 
 ARGS=$@
 
-set -e
-
 SCRIPT_PATH=$0
 if [ -L "$SCRIPT_PATH" ]; then
     SCRIPT_PATH=$(dirname "$SCRIPT_PATH")/$(readlink "$SCRIPT_PATH")
@@ -24,7 +22,7 @@ RESULT=0
 source "$BASE_PATH/SetupMono.sh" $BASE_PATH
 
 while : ; do
-        mono Engine/Binaries/DotNET/GitDependencies.exe $ARGS && break
+        mono Engine/Binaries/DotNET/GitDependencies.exe $ARGS
         RESULT=$?
 
         echo "Result: $RESULT"
