Process needs to be set in the background so the SIG handler can work for RunUAT.sh

diff --git a/Engine/Build/BatchFiles/RunUAT.sh b/Engine/Build/BatchFiles/RunUAT.sh
--- a/Engine/Build/BatchFiles/RunUAT.sh
+++ b/Engine/Build/BatchFiles/RunUAT.sh
@@ -11,7 +11,7 @@ echo Running AutomationTool...
 echo
 
 TermHandler() {
-	kill -TERM $UATPid 2>&1 /dev/null
+	kill -TERM $UATPid 2> /dev/null
 }
 
 # loop over the arguments, quoting spaces to pass to UAT proper
@@ -101,7 +101,7 @@ trap TermHandler SIGTERM SIGINT
 # you can't set a dotted env var nicely in sh, but env will run a command with
 # a list of env vars set, including dotted ones
 echo Start UAT: mono AutomationTool.exe "${Args[@]}"
-env uebp_LogFolder="$LogDir" mono AutomationTool.exe "${Args[@]}" $UATCompileArg
+env uebp_LogFolder="$LogDir" mono AutomationTool.exe "${Args[@]}" $UATCompileArg &
 UATPid=$!
 wait $UATPid
 UATReturn=$?
