Fix for outof BatchFiles/Linux usage

diff --git a/Engine/Build/BatchFiles/Linux/GitDependencies.sh b/Engine/Build/BatchFiles/Linux/GitDependencies.sh
--- a/Engine/Build/BatchFiles/Linux/GitDependencies.sh
+++ b/Engine/Build/BatchFiles/Linux/GitDependencies.sh
@@ -6,10 +6,19 @@
 
 ARGS=$@
 
-BASE_PATH="`pwd`/`dirname "$0"`"
+set -e
+
+SCRIPT_PATH=$0
+if [ -L "$SCRIPT_PATH" ]; then
+    SCRIPT_PATH=$(dirname "$SCRIPT_PATH")/$(readlink "$SCRIPT_PATH")
+fi
+
+cd "$(dirname "$SCRIPT_PATH")" && SCRIPT_PATH="`pwd`/$(basename "$SCRIPT_PATH")"
+
+BASE_PATH="`dirname "$SCRIPT_PATH"`"
 
 # cd to Engine root
-cd "$(dirname "$BASH_SOURCE")"/../../../..
+cd ../../../..
 RESULT=0
 
 source "$BASE_PATH/SetupMono.sh" $BASE_PATH
