Remove the automatically appended -DEPLOY argument from Build.bat. Build platforms already enable this if they want it; it just causes unnecessary log spam for platforms that don't.

#rnx

diff --git a/Engine/Build/BatchFiles/Build.bat b/Engine/Build/BatchFiles/Build.bat
--- a/Engine/Build/BatchFiles/Build.bat
+++ b/Engine/Build/BatchFiles/Build.bat
@@ -11,7 +11,7 @@ REM %2 is the platform name
 REM %3 is the configuration name
 
 IF EXIST ..\..\Engine\Binaries\DotNET\UnrealBuildTool.exe (
-        ..\..\Engine\Binaries\DotNET\UnrealBuildTool.exe %* -DEPLOY
+        ..\..\Engine\Binaries\DotNET\UnrealBuildTool.exe %*
 		popd
 
 		REM Ignore exit codes of 2 ("ECompilationResult.UpToDate") from UBT; it's not a failure.
