Make FixMonoFiles overwrite symlinks whose targets do not exist.

#rnx

diff --git a/Engine/Build/BatchFiles/Mac/FixMonoFiles.sh b/Engine/Build/BatchFiles/Mac/FixMonoFiles.sh
--- a/Engine/Build/BatchFiles/Mac/FixMonoFiles.sh
+++ b/Engine/Build/BatchFiles/Mac/FixMonoFiles.sh
@@ -3,7 +3,7 @@
 CUR_DIR=`pwd`
 cd ../../../Binaries/ThirdParty/Mono/Mac
 
-SCRIPT_VERSION=1
+SCRIPT_VERSION=2
 
 UPDATE_LINKS=false
 if [ ! -f FixMonoFiles.version ]; then
@@ -14,7 +14,7 @@ fi
 
 function try_symlink {
 	if [ -f $(dirname $2)/$1 ]; then
-		if [ ! -f $2 ]; then
+		if [ ! -L $2 ]; then
 			ln -s $1 $2
 		elif [ "$UPDATE_LINKS" = true ]; then
 			rm $2
