Fix warning when generating project files in a minimal workspace (due to no documentation folder).

diff --git a/Engine/Build/BatchFiles/FixDependencyFiles.bat b/Engine/Build/BatchFiles/FixDependencyFiles.bat
--- a/Engine/Build/BatchFiles/FixDependencyFiles.bat
+++ b/Engine/Build/BatchFiles/FixDependencyFiles.bat
@@ -2,5 +2,9 @@
 pushd %~dp0..\..\
 for /f "tokens=*" %%X in ('dir /b /s /a:-d Content\Editor\Slate\*.PNG') do ren "%%X" "%%~nX.png"
 for /f "tokens=*" %%X in ('dir /b /s /a:-d Content\Slate\*.PNG') do ren "%%X" "%%~nX.png"
+
+if not exist Documentation\Source\Shared\Icons goto no_documentation_icons
 for /f "tokens=*" %%X in ('dir /b /s /a:-d Documentation\Source\Shared\Icons\*.PNG') do ren "%%X" "%%~nX.png"
+:no_documentation_icons
+
 popd
\ No newline at end of file
