Linux: fix which spam on some systems (UE-69994).

diff --git a/Engine/Build/BatchFiles/Linux/SetupMono.sh b/Engine/Build/BatchFiles/Linux/SetupMono.sh
--- a/Engine/Build/BatchFiles/Linux/SetupMono.sh
+++ b/Engine/Build/BatchFiles/Linux/SetupMono.sh
@@ -11,7 +11,7 @@ bash FixMonoFiles.sh
 bash FixDependencyFiles.sh
 
 IS_MONO_INSTALLED=0
-MONO_VERSION_PATH=`which mono` || true
+MONO_VERSION_PATH=$(command -v mono) || true
 if [ "$UE_USE_SYSTEM_MONO" == "1" ] && [ ! $MONO_VERSION_PATH == "" ] && [ -f $MONO_VERSION_PATH ]; then
 	# If Mono is installed, check if it's 4.0.2 or higher
 	MONO_VERSION_PREFIX="Mono JIT compiler version "
