Add a dedicated file to check whether binary prerequisites have been downloaded from GitHub, rather than checking for SwarmAgent.exe or any other build product.

diff --git a/Engine/Build/BatchFiles/GenerateProjectFiles.bat b/Engine/Build/BatchFiles/GenerateProjectFiles.bat
--- a/Engine/Build/BatchFiles/GenerateProjectFiles.bat
+++ b/Engine/Build/BatchFiles/GenerateProjectFiles.bat
@@ -22,7 +22,7 @@ if not exist ..\Build\BatchFiles\GenerateProjectFiles.bat goto Error_BatchFileIn
 
 rem ## Check to make sure that we have a Binaries directory with at least one dependency that we know that UnrealBuildTool will need
 rem ## in order to run.  It's possible the user acquired source but did not download and unpack the other prerequiste binaries.
-if not exist ..\Binaries\DotNET\SwarmAgent.exe goto Error_MissingBinaryPrerequisites
+if not exist ..\Build\BinaryPrerequisitesMarker.dat goto Error_MissingBinaryPrerequisites
 
 rem ## Get the path to MSBuild
 call "%~dp0GetMSBuildPath.bat"
