Fix failure to compile C# projects when only VS2019 is installed. MSBuild has been moved into a different directory.

diff --git a/Engine/Build/BatchFiles/GetMSBuildPath.bat b/Engine/Build/BatchFiles/GetMSBuildPath.bat
--- a/Engine/Build/BatchFiles/GetMSBuildPath.bat
+++ b/Engine/Build/BatchFiles/GetMSBuildPath.bat
@@ -8,9 +8,13 @@ rem ## The discovered path is set to the MSBUILD_EXE environment variable on suc
 
 set MSBUILD_EXE=
 
-rem ## Try to get the MSBuild 15 path using vswhere (see [REDACTED_PROFILE_URL]/vswhere)
+rem ## Try to get the MSBuild 15 path using vswhere (see [REDACTED_PROFILE_URL]/vswhere). VS2019 preview puts the executable in a "Current" subfolder.
 if not exist "%ProgramFiles(x86)%\Microsoft Visual Studio\Installer\vswhere.exe" goto no_vswhere
 for /f "delims=" %%i in ('"%ProgramFiles(x86)%\Microsoft Visual Studio\Installer\vswhere" -latest -products * -requires Microsoft.Component.MSBuild -property installationPath') do (
+	if exist "%%i\MSBuild\Current\Bin\MSBuild.exe" (
+		set MSBUILD_EXE="%%i\MSBuild\Current\Bin\MSBuild.exe"
+		goto Succeeded
+	)
 	if exist "%%i\MSBuild\15.0\Bin\MSBuild.exe" (
 		set MSBUILD_EXE="%%i\MSBuild\15.0\Bin\MSBuild.exe"
 		goto Succeeded
@@ -40,9 +44,6 @@ rem ## Check for older versions of MSBuild. These are registered as separate ver
 call :ReadInstallPath Microsoft\MSBuild\ToolsVersions\12.0 MSBuildToolsPath MSBuild.exe
 if not errorlevel 1 goto Succeeded
 
-call :ReadInstallPath Microsoft\MSBuild\ToolsVersions\4.0 MSBuildToolsPath MSBuild.exe
-if not errorlevel 1 goto Succeeded
-
 rem ## Couldn't find anything
 exit /B 1
 
