Allow pre-release versions of Visual Studio when getting path to MSBuild.

diff --git a/Engine/Build/BatchFiles/GetMSBuildPath.bat b/Engine/Build/BatchFiles/GetMSBuildPath.bat
--- a/Engine/Build/BatchFiles/GetMSBuildPath.bat
+++ b/Engine/Build/BatchFiles/GetMSBuildPath.bat
@@ -10,7 +10,7 @@ set MSBUILD_EXE=
 
 rem ## Try to get the MSBuild 15 path using vswhere (see [REDACTED_PROFILE_URL]/vswhere). VS2019 preview puts the executable in a "Current" subfolder.
 if not exist "%ProgramFiles(x86)%\Microsoft Visual Studio\Installer\vswhere.exe" goto no_vswhere
-for /f "delims=" %%i in ('"%ProgramFiles(x86)%\Microsoft Visual Studio\Installer\vswhere" -latest -products * -requires Microsoft.Component.MSBuild -property installationPath') do (
+for /f "delims=" %%i in ('"%ProgramFiles(x86)%\Microsoft Visual Studio\Installer\vswhere" -prerelease -latest -products * -requires Microsoft.Component.MSBuild -property installationPath') do (
 	if exist "%%i\MSBuild\Current\Bin\MSBuild.exe" (
 		set MSBUILD_EXE="%%i\MSBuild\Current\Bin\MSBuild.exe"
 		goto Succeeded
