- Add check for installedbuild.txt to force -compile to off
- Remove deletion of UAT intermediates if we are doing a -compile to match the .bat behavior

diff --git a/Engine/Build/BatchFiles/RunUAT.sh b/Engine/Build/BatchFiles/RunUAT.sh
--- a/Engine/Build/BatchFiles/RunUAT.sh
+++ b/Engine/Build/BatchFiles/RunUAT.sh
@@ -51,6 +51,10 @@ else
 	UATCompileArg=-compile
 fi
 
+if [ -f Build\InstalledBuild.txt ]; then
+	UATCompileArg=
+fi
+
 if [ "$(uname)" = "Darwin" ]; then
 	# Setup Mono
 	source "$SCRIPT_DIR/Mac/SetupMono.sh" "$SCRIPT_DIR/Mac"
@@ -67,14 +71,6 @@ if [ "$UATCompileArg" = "-compile" ]; then
 		echo No project to compile, attempting to use precompiled AutomationTool
 		UATCompileArg=
 	else
-		if [ -f "$UATDirectory/AutomationTool.exe" ]; then
-			echo AutomationTool exists: Deleting
-			rm -f $UATDirectory/AutomationTool.exe
-			if [ -d "$UATDirectory/AutomationScripts" ]; then
-				echo Deleting all AutomationScript dlls
-				rm -rf $UATDirectory/AutomationScripts
-			fi
-		fi
 		echo Compiling AutomationTool with xbuild
 
 		ARGS="/p:Configuration=Development /p:Platform=AnyCPU /verbosity:quiet /nologo"
