Improve error message when GenerateProjectFiles.bat fails due to no Visual Studio installation having an MSBuild installation available.

diff --git a/Engine/Build/BatchFiles/GenerateProjectFiles.bat b/Engine/Build/BatchFiles/GenerateProjectFiles.bat
--- a/Engine/Build/BatchFiles/GenerateProjectFiles.bat
+++ b/Engine/Build/BatchFiles/GenerateProjectFiles.bat
@@ -99,7 +99,7 @@ goto Exit
 
 :Error_NoVisualStudioEnvironment
 echo.
-echo GenerateProjectFiles ERROR: We couldn't find a valid installation of Visual Studio.  This program requires Visual Studio 2017.  Please check that you have Visual Studio installed, then verify that the HKEY_LOCAL_MACHINE\SOFTWARE\Wow6432Node\Microsoft\VisualStudio\14.0\InstallDir (or HKEY_LOCAL_MACHINE\SOFTWARE\Microsoft\VisualStudio\14.0\InstallDir on 32-bit machines) registry value is set.  Visual Studio configures this value when it is installed, and this program expects it to be set to the '\Common7\IDE\' sub-folder under a valid Visual Studio installation directory.
+echo GenerateProjectFiles ERROR: Unable to find a valid installation of Visual Studio.  Please check that you have Visual Studio 2017 or Visual Studio 2019 installed, and the MSBuild component is selected as part of your installation.
 echo.
 pause
 goto Exit
