- Last of the Dev-Build errors (at least for now)

diff --git a/Engine/Build/BatchFiles/FindPlatformExtensionSources.bat b/Engine/Build/BatchFiles/FindPlatformExtensionSources.bat
--- a/Engine/Build/BatchFiles/FindPlatformExtensionSources.bat
+++ b/Engine/Build/BatchFiles/FindPlatformExtensionSources.bat
@@ -1,4 +1,10 @@
-setlocal
+@echo off
+rem ## Unreal Engine 4 AutomationTool setup script
+rem ## Copyright 1998-2019 Epic Games, Inc. All Rights Reserved.
+rem ##
+rem ## This script is expecting to exist in the UE4/Engine/Build/BatchFiles directory.  It will not work correctly
+rem ## if you copy it to a different location and run it.
+
 pushd "%~dp0..\..\Source"
 
 mkdir ..\Intermediate\ProjectFiles >NUL 2>NUL
