- Added support for Mac/Linux to have PlatformExtensions sources for UBT (FindPlatformExtensionsSources.sh)

diff --git a/Engine/Build/BatchFiles/FindPlatformExtensionSources.sh b/Engine/Build/BatchFiles/FindPlatformExtensionSources.sh
new file mode 100755
--- /dev/null
+++ b/Engine/Build/BatchFiles/FindPlatformExtensionSources.sh
@@ -0,0 +1,29 @@
+#!/bin/bash
+## Copyright 1998-2019 Epic Games, Inc. All Rights Reserved.
+##
+## Unreal Engine 4 AutomationTool setup script
+##
+## This script is expecting to exist in the UE4/Engine/Build/BatchFiles directory.  It will not work
+## correctly if you copy it to a different location and run it.
+
+
+SCRIPT_DIR=$(cd "`dirname "$0"`" && pwd)
+cd "$SCRIPT_DIR/../../Source"
+
+mkdir -p ../Intermediate/ProjectFiles
+
+## Look for any platform extension .cs files, and add them to a file that will be refernced by UBT project, so it can bring them in automatically
+REFERENCE_FILE=../Intermediate/ProjectFiles/UnrealBuildTool.csproj.References
+echo "<?xml version=\"1.0\" encoding=\"utf-8\"?>" > $REFERENCE_FILE
+echo "<Project ToolsVersion=\"15.0\" DefaultTargets=\"Build\" xmlns=\"http://schemas.microsoft.com/developer/msbuild/2003\">" >> $REFERENCE_FILE
+echo "<ItemGroup>" >> $REFERENCE_FILE
+
+for PLATFORM in ../../Platforms/*; do
+	for SOURCE_FILE in ${PLATFORM}/Engine/Source/Programs/UnrealBuildTool/*; do
+		echo "<PlatformExtensionCompile Include=\"..\\..\\${SOURCE_FILE//\//\\}\" />" >> $REFERENCE_FILE
+	done
+done
+
+echo "</ItemGroup>" >> $REFERENCE_FILE
+echo "</Project>" >> $REFERENCE_FILE
+
diff --git a/Engine/Build/BatchFiles/Linux/GenerateProjectFiles.sh b/Engine/Build/BatchFiles/Linux/GenerateProjectFiles.sh
--- a/Engine/Build/BatchFiles/Linux/GenerateProjectFiles.sh
+++ b/Engine/Build/BatchFiles/Linux/GenerateProjectFiles.sh
@@ -22,6 +22,9 @@ fi
 
 source "$BASE_PATH/SetupMono.sh" $BASE_PATH
 
+# make sure the UBT project has references to auto-discovered platform extension source files
+"${BASE_PATH}/../FindPlatformExtensionSources.sh"
+
 if [ -f "$BASE_PATH/../../../Source/Programs/UnrealBuildTool/UnrealBuildTool.csproj" ]; then
 	xbuild "$BASE_PATH/../../../Source/Programs/UnrealBuildTool/UnrealBuildTool.csproj" /property:Configuration="Development" /verbosity:quiet /nologo
 fi
diff --git a/Engine/Build/BatchFiles/Mac/GenerateProjectFiles.sh b/Engine/Build/BatchFiles/Mac/GenerateProjectFiles.sh
--- a/Engine/Build/BatchFiles/Mac/GenerateProjectFiles.sh
+++ b/Engine/Build/BatchFiles/Mac/GenerateProjectFiles.sh
@@ -21,6 +21,9 @@ fi
 
 . "$BASE_PATH/SetupMono.sh" "$BASE_PATH"
 
+# make sure the UBT project has references to auto-discovered platform extension source files
+"${BASE_PATH}/../FindPlatformExtensionSources.sh"
+
 if [ -f "$BASE_PATH/../../../Source/Programs/UnrealBuildTool/UnrealBuildTool.csproj" ]; then
 	xbuild "$BASE_PATH/../../../Source/Programs/UnrealBuildTool/UnrealBuildTool.csproj" /property:Configuration="Development" /verbosity:quiet /nologo |grep -i error
 fi
diff --git a/Engine/Build/BatchFiles/RunUAT.sh b/Engine/Build/BatchFiles/RunUAT.sh
--- a/Engine/Build/BatchFiles/RunUAT.sh
+++ b/Engine/Build/BatchFiles/RunUAT.sh
@@ -73,6 +73,9 @@ if [ "$UATCompileArg" = "-compile" ]; then
 	else
 		echo Compiling AutomationTool with xbuild
 
+		# make sure the UBT project has references to auto-discovered platform extension source files
+		"${SCRIPT_DIR}/FindPlatformExtensionSources.sh"
+
 		ARGS="/p:Configuration=Development /p:Platform=AnyCPU /verbosity:quiet /nologo"
 		ARGS="${ARGS} /p:TargetFrameworkProfile="
 
