Fixed issue in batch scripts were directory to batch script was not quoted and as such would cause issue when special characters was part of the path.

diff --git a/Engine/Build/BatchFiles/FixDependencyFiles.bat b/Engine/Build/BatchFiles/FixDependencyFiles.bat
--- a/Engine/Build/BatchFiles/FixDependencyFiles.bat
+++ b/Engine/Build/BatchFiles/FixDependencyFiles.bat
@@ -1,5 +1,5 @@
 @echo off
-pushd %~dp0..\..\
+pushd "%~dp0..\..\"
 for /f "tokens=*" %%X in ('dir /b /s /a:-d Content\Editor\Slate\*.PNG') do ren "%%X" "%%~nX.png"
 for /f "tokens=*" %%X in ('dir /b /s /a:-d Content\Slate\*.PNG') do ren "%%X" "%%~nX.png"
 
diff --git a/Engine/Build/BatchFiles/RunUAT.bat b/Engine/Build/BatchFiles/RunUAT.bat
--- a/Engine/Build/BatchFiles/RunUAT.bat
+++ b/Engine/Build/BatchFiles/RunUAT.bat
@@ -14,7 +14,7 @@ set UATDirectory=Binaries\DotNET\
 set UATCompileArg=-compile
 
 rem ## Change the CWD to /Engine. 
-pushd %~dp0..\..\
+pushd "%~dp0..\..\"
 if not exist Build\BatchFiles\RunUAT.bat goto Error_BatchFileInWrongLocation
 
 rem ## Use the pre-compiled UAT scripts if -nocompile is specified in the command line
diff --git a/Setup.bat b/Setup.bat
--- a/Setup.bat
+++ b/Setup.bat
@@ -1,6 +1,6 @@
 @echo off
 setlocal
-pushd %~dp0
+pushd "%~dp0"
 
 rem Figure out if we should append the -prompt argument
 set PROMPT_ARGUMENT=
