Put C# files for platform extensions under a platform subfolder.

#rnx

diff --git a/Engine/Build/BatchFiles/FindPlatformExtensionSources.bat b/Engine/Build/BatchFiles/FindPlatformExtensionSources.bat
--- a/Engine/Build/BatchFiles/FindPlatformExtensionSources.bat
+++ b/Engine/Build/BatchFiles/FindPlatformExtensionSources.bat
@@ -12,17 +12,25 @@ pushd "%~dp0..\..\Source"
 mkdir ..\Intermediate\ProjectFiles >NUL 2>NUL
 
 rem Look for any platform extension .cs files, and add them to a file that will be refernced by UBT project, so it can bring them in automatically
-set REFERENCE_FILE="..\Intermediate\ProjectFiles\UnrealBuildTool.csproj.References"
+set REFERENCE_FILE="%~dp0..\..\Intermediate\ProjectFiles\UnrealBuildTool.csproj.References"
 echo ^<?xml version="1.0" encoding="utf-8"?^> > %REFERENCE_FILE%
 echo ^<Project ToolsVersion="15.0" DefaultTargets="Build" xmlns="http://schemas.microsoft.com/developer/msbuild/2003"^> >> %REFERENCE_FILE%
 echo ^<ItemGroup^> >> %REFERENCE_FILE%
 
 if exist "..\..\Platforms" (
-	for /d %%i in ("..\..\Platforms\*") do (
-		for %%j in (%%i\Engine\Source\Programs\UnrealBuildTool\*) do (
-			echo ^<PlatformExtensionCompile Include="..\..\%%j" /^> >> %REFERENCE_FILE%
+	pushd "..\..\Platforms"
+	for /d %%i in ("*") do (
+		if exist "%%i\Engine\Source\Programs\UnrealBuildTool" (
+			pushd "%%i\Engine\Source\Programs\UnrealBuildTool"
+			for %%j in ("*.cs") do (
+				echo   ^<PlatformExtensionCompile Include="..\..\..\..\Platforms\%%i\Engine\Source\Programs\UnrealBuildTool\%%j"^> >> %REFERENCE_FILE%
+				echo     ^<Link^>Platform\%%i\%%j^</Link^> >> %REFERENCE_FILE%
+				echo   ^</PlatformExtensionCompile^> >> %REFERENCE_FILE%
+			)
+			popd
 		)
 	)
+	popd
 )
 echo ^</ItemGroup^> >> %REFERENCE_FILE%
 echo ^</Project^> >> %REFERENCE_FILE%
\ No newline at end of file
