Fixed the error when checking a default value that doesn't exist on Mac for generation of project files

diff --git a/Engine/Build/BatchFiles/Mac/GenerateProjectFiles.sh b/Engine/Build/BatchFiles/Mac/GenerateProjectFiles.sh
--- a/Engine/Build/BatchFiles/Mac/GenerateProjectFiles.sh
+++ b/Engine/Build/BatchFiles/Mac/GenerateProjectFiles.sh
@@ -30,40 +30,42 @@ fi
 
 WANT_AOT="`defaults read com.epicgames.ue4 MonoAOT`"
 OPENSSL="./../../../Binaries/DotNET/IOS/openssl.exe"
-if [ $WANT_AOT == "1" ]; then
-        for i in $BASE_PATH/../../../Binaries/DotNET/*.dll;
-        do
-            if test "$i" -nt "$i.dylib"; then
-        		echo Compiling $i to native...
-                mono --aot $i > /dev/null 2>&1;
-            fi
-        done
+if [ ! -z $WANT_AOT ]; then
+	if [ $WANT_AOT == "1" ]; then
+			for i in $BASE_PATH/../../../Binaries/DotNET/*.dll;
+			do
+				if test "$i" -nt "$i.dylib"; then
+					echo Compiling $i to native...
+					mono --aot $i > /dev/null 2>&1;
+				fi
+			done
 
-        for i in $BASE_PATH/../../../Binaries/DotNET/*.exe;
-        do
-            if test "$i" -nt "$i.dylib"; then
-        		echo Compiling $i to native...
-                mono --aot $i > /dev/null 2>&1;
-            fi
-        done
+			for i in $BASE_PATH/../../../Binaries/DotNET/*.exe;
+			do
+				if test "$i" -nt "$i.dylib"; then
+					echo Compiling $i to native...
+					mono --aot $i > /dev/null 2>&1;
+				fi
+			done
 
-        for i in $BASE_PATH/../../../Binaries/DotNET/IOS/*.dll;
-        do
-            if test "$i" -nt "$i.dylib"; then
-        		echo Compiling $i to native...
-                mono --aot $i > /dev/null 2>&1;
-            fi
-        done
+			for i in $BASE_PATH/../../../Binaries/DotNET/IOS/*.dll;
+			do
+				if test "$i" -nt "$i.dylib"; then
+					echo Compiling $i to native...
+					mono --aot $i > /dev/null 2>&1;
+				fi
+			done
 
-        for i in $BASE_PATH/../../../Binaries/DotNET/IOS/*.exe;
-        do
-            if test "$i" -nt "$i.dylib"; then
-                if [ $i != $OPENSSL ]; then
-                    echo Compiling $i to native...
-                    mono --aot $i > /dev/null 2>&1;
-                fi
-            fi
-        done
+			for i in $BASE_PATH/../../../Binaries/DotNET/IOS/*.exe;
+			do
+				if test "$i" -nt "$i.dylib"; then
+					if [ $i != $OPENSSL ]; then
+						echo Compiling $i to native...
+						mono --aot $i > /dev/null 2>&1;
+					fi
+				fi
+			done
+	fi
 fi
 
 # pass all parameters to UBT
