Prevent UAT being compiled in installed builds.

diff --git a/Engine/Build/BatchFiles/RunUAT.bat b/Engine/Build/BatchFiles/RunUAT.bat
--- a/Engine/Build/BatchFiles/RunUAT.bat
+++ b/Engine/Build/BatchFiles/RunUAT.bat
@@ -20,6 +20,9 @@ if not exist Build\BatchFiles\RunUAT.bat goto Error_BatchFileInWrongLocation
 rem ## Use the pre-compiled UAT scripts if -nocompile is specified in the command line
 for %%P in (%*) do if /I "%%P" == "-nocompile" goto RunPrecompiled
 
+rem ## If we're running in an installed build, default to precompiled
+if exist Build\InstalledBuild.txt goto RunPrecompiled
+
 rem ## check for force precompiled
 if not "%ForcePrecompiledUAT%"=="" goto RunPrecompiled
 
