Linux: do not install unnecessary dependencies.

- Also fixes building on Ubuntu 18.10.

diff --git a/Engine/Build/BatchFiles/Linux/Setup.sh b/Engine/Build/BatchFiles/Linux/Setup.sh
--- a/Engine/Build/BatchFiles/Linux/Setup.sh
+++ b/Engine/Build/BatchFiles/Linux/Setup.sh
@@ -22,97 +22,12 @@ if [ -e /etc/os-release ]; then
   source /etc/os-release
   # Ubuntu/Debian/Mint
   if [[ "$ID" == "ubuntu" ]] || [[ "$ID_LIKE" == "ubuntu" ]] || [[ "$ID" == "debian" ]] || [[ "$ID_LIKE" == "debian" ]] || [[ "$ID" == "tanglu" ]] || [[ "$ID_LIKE" == "tanglu" ]]; then
-    # Install the necessary dependencies (require clang-3.8 on 16.04, although 3.3 and 3.5 through 3.7 should work too for this release)
-     # mono-devel is needed for making the installed build (particularly installing resgen2 tool)
 
-    # Adjust the VERSION_ID for elementary OS since elementary doesn't follow Debian versioning, but is Ubuntu-based.
-    if [[ "$ID" == "elementary" ]]; then
-      if [[ "$VERSION_ID" == 0.4 ]] || ([[ "$VERSION_ID" > 0.4 ]] && [[ "$VERSION_ID" < 0.5 ]]); then
-        VERSION_ID=16.04 # The 0.4 branch is based on 16.04
-      elif [[ "$VERSION_ID" == 0.5 ]] || [[ "$VERSION_ID" > 0.5 ]]; then
-        VERSION_ID=17 # Assumes that 0.5 won't be 16.04 based
-      fi
-    fi
-
-    if [ -n "$VERSION_ID" ] && [[ "$VERSION_ID" < 16.04 ]]; then
-     DEPS="mono-xbuild \
-       mono-dmcs \
-       libmono-microsoft-build-tasks-v4.0-4.0-cil \
-       libmono-system-data-datasetextensions4.0-cil
-       libmono-system-web-extensions4.0-cil
-       libmono-system-management4.0-cil
-       libmono-system-xml-linq4.0-cil
-       libmono-corlib4.0-cil
-       libmono-windowsbase4.0-cil
-       libmono-system-io-compression4.0-cil
-       libmono-system-io-compression-filesystem4.0-cil
-       mono-devel
-       build-essential
-       "
-    elif [ -n "$VERSION_ID" ] && [[ "$VERSION_ID" == 16.04 ]]; then
-     DEPS="mono-xbuild \
-       mono-dmcs \
-       libmono-microsoft-build-tasks-v4.0-4.0-cil \
-       libmono-system-data-datasetextensions4.0-cil
-       libmono-system-web-extensions4.0-cil
-       libmono-system-management4.0-cil
-       libmono-system-xml-linq4.0-cil
-       libmono-corlib4.5-cil
-       libmono-windowsbase4.0-cil
-       libmono-system-io-compression4.0-cil
-       libmono-system-io-compression-filesystem4.0-cil
-       libmono-system-runtime4.0-cil
-       mono-devel
-       build-essential
-       "
-    elif [ -n "$VERSION_ID" ] && [[ "$VERSION_ID" < 17.10 ]]; then
-     DEPS="mono-xbuild \
-       mono-dmcs \
-       libmono-microsoft-build-tasks-v4.0-4.0-cil \
-       libmono-system-data-datasetextensions4.0-cil
-       libmono-system-web-extensions4.0-cil
-       libmono-system-management4.0-cil
-       libmono-system-xml-linq4.0-cil
-       libmono-corlib4.5-cil
-       libmono-windowsbase4.0-cil
-       libmono-system-io-compression4.0-cil
-       libmono-system-io-compression-filesystem4.0-cil
-       libmono-system-runtime4.0-cil
-       mono-devel
-       build-essential
-       "
-    elif [[ $PRETTY_NAME == *sid ]] || [[ $PRETTY_NAME == *stretch ]]; then
-     DEPS="mono-xbuild \
-       mono-dmcs \
-       libmono-microsoft-build-tasks-v4.0-4.0-cil \
-       libmono-system-data-datasetextensions4.0-cil
-       libmono-system-web-extensions4.0-cil
-       libmono-system-management4.0-cil
-       libmono-system-xml-linq4.0-cil
-       libmono-corlib4.5-cil
-       libmono-windowsbase4.0-cil
-       libmono-system-io-compression4.0-cil
-       libmono-system-io-compression-filesystem4.0-cil
-       libmono-system-runtime4.0-cil
-       mono-devel
-       "
-    else # assume the latest Ubuntu, this is going to be a moving target (17.10 as of now)
-     DEPS="mono-xbuild \
-       mono-dmcs \
-       libmono-microsoft-build-tasks-v4.0-4.0-cil \
-       libmono-system-data-datasetextensions4.0-cil
-       libmono-system-web-extensions4.0-cil
-       libmono-system-management4.0-cil
-       libmono-system-xml-linq4.0-cil
-       libmono-corlib4.5-cil
-       libmono-windowsbase4.0-cil
-       libmono-system-io-compression4.0-cil
-       libmono-system-io-compression-filesystem4.0-cil
-       libmono-system-runtime4.0-cil
-       mono-devel
-       build-essential
-       "
-    fi
+    # Both mono and the compiler toolchain are now bundled.
+    # build-essential can still be useful for making sure the system has make
+    DEPS="
+      build-essential
+      "
 
     for DEP in $DEPS; do
       if ! PackageIsInstalled $DEP; then
