Prevent FixMonoFiles.version from being updated unnecessarily.

diff --git a/Engine/Build/BatchFiles/Mac/FixMonoFiles.sh b/Engine/Build/BatchFiles/Mac/FixMonoFiles.sh
--- a/Engine/Build/BatchFiles/Mac/FixMonoFiles.sh
+++ b/Engine/Build/BatchFiles/Mac/FixMonoFiles.sh
@@ -174,6 +174,8 @@ try_chmod bin/mcs
 try_chmod bin/mono-sgen64
 try_chmod bin/xbuild
 
-echo "$SCRIPT_VERSION" > FixMonoFiles.version
+if [ "$UPDATE_LINKS" = true ] ; then
+	echo "$SCRIPT_VERSION" > FixMonoFiles.version
+fi
 
 cd "$CUR_DIR"
