Suppress message when fixing Slate png case. Users don't need to know about this.

#rnx

diff --git a/Engine/Build/BatchFiles/FixDependencyFiles.bat b/Engine/Build/BatchFiles/FixDependencyFiles.bat
--- a/Engine/Build/BatchFiles/FixDependencyFiles.bat
+++ b/Engine/Build/BatchFiles/FixDependencyFiles.bat
@@ -1,5 +1,3 @@
-echo Fixing inconsistent case in filenames.
-
 @echo off
 pushd %~dp0..\..\
 for /f "tokens=*" %%X in ('dir /b /s /a:-d Content\Editor\Slate\*.PNG') do ren "%%X" "%%~nX.png"
