Fix error messages that instruct the user to install Visual Studio 2015 (should be 2017)

#rnx

diff --git a/Engine/Build/BatchFiles/GenerateProjectFiles.bat b/Engine/Build/BatchFiles/GenerateProjectFiles.bat
--- a/Engine/Build/BatchFiles/GenerateProjectFiles.bat
+++ b/Engine/Build/BatchFiles/GenerateProjectFiles.bat
@@ -96,7 +96,7 @@ goto Exit
 
 :Error_NoVisualStudioEnvironment
 echo.
-echo GenerateProjectFiles ERROR: We couldn't find a valid installation of Visual Studio.  This program requires Visual Studio 2015.  Please check that you have Visual Studio installed, then verify that the HKEY_LOCAL_MACHINE\SOFTWARE\Wow6432Node\Microsoft\VisualStudio\14.0\InstallDir (or HKEY_LOCAL_MACHINE\SOFTWARE\Microsoft\VisualStudio\14.0\InstallDir on 32-bit machines) registry value is set.  Visual Studio configures this value when it is installed, and this program expects it to be set to the '\Common7\IDE\' sub-folder under a valid Visual Studio installation directory.
+echo GenerateProjectFiles ERROR: We couldn't find a valid installation of Visual Studio.  This program requires Visual Studio 2017.  Please check that you have Visual Studio installed, then verify that the HKEY_LOCAL_MACHINE\SOFTWARE\Wow6432Node\Microsoft\VisualStudio\14.0\InstallDir (or HKEY_LOCAL_MACHINE\SOFTWARE\Microsoft\VisualStudio\14.0\InstallDir on 32-bit machines) registry value is set.  Visual Studio configures this value when it is installed, and this program expects it to be set to the '\Common7\IDE\' sub-folder under a valid Visual Studio installation directory.
 echo.
 pause
 goto Exit
diff --git a/Engine/Source/Programs/UnrealBuildTool/Platform/Windows/UEBuildWindows.cs b/Engine/Source/Programs/UnrealBuildTool/Platform/Windows/UEBuildWindows.cs
--- a/Engine/Source/Programs/UnrealBuildTool/Platform/Windows/UEBuildWindows.cs
+++ b/Engine/Source/Programs/UnrealBuildTool/Platform/Windows/UEBuildWindows.cs
@@ -687,7 +687,7 @@ namespace UnrealBuildTool
 			}
 			else
 			{
-				Log.TraceWarning("No Visual C++ installation was found. Please download and install Visual Studio 2015 with C++ components.");
+				Log.TraceWarning("No Visual C++ installation was found. Please download and install Visual Studio 2017 with C++ components.");
 			}
 
 			// Finally, default to VS2017 anyway
