PR #5615: Fix quotes on SetupMono.sh for Mac (Contributed by [REDACTED_PERSON])

diff --git a/Engine/Build/BatchFiles/Mac/GenerateProjectFiles.sh b/Engine/Build/BatchFiles/Mac/GenerateProjectFiles.sh
--- a/Engine/Build/BatchFiles/Mac/GenerateProjectFiles.sh
+++ b/Engine/Build/BatchFiles/Mac/GenerateProjectFiles.sh
@@ -19,7 +19,7 @@ if [ ! -d "$BASE_PATH/../../../Source" ]; then
  exit 1
 fi
 
-. "$BASE_PATH/SetupMono.sh" $BASE_PATH
+. "$BASE_PATH/SetupMono.sh" "$BASE_PATH"
 
 if [ -f "$BASE_PATH/../../../Source/Programs/UnrealBuildTool/UnrealBuildTool.csproj" ]; then
 	xbuild "$BASE_PATH/../../../Source/Programs/UnrealBuildTool/UnrealBuildTool.csproj" /property:Configuration="Development" /verbosity:quiet /nologo |grep -i error
