Fix check for whether mono >= 5 is installed

diff --git a/Engine/Build/BatchFiles/RunUAT.sh b/Engine/Build/BatchFiles/RunUAT.sh
--- a/Engine/Build/BatchFiles/RunUAT.sh
+++ b/Engine/Build/BatchFiles/RunUAT.sh
@@ -77,7 +77,7 @@ if [ "$UATCompileArg" = "-compile" ]; then
 		"${SCRIPT_DIR}/FindPlatformExtensionSources.sh"
 
 		# mono 5.0 and up include msbuild
-		if [ ! $IS_MONO_5_INSTALLED == "" ]; then
+		if [ "$IS_MONO_5_INSTALLED" == "1" ]; then
 			BUILD_TOOL=msbuild
 		else
 			BUILD_TOOL=xbuild
