Require that installed mono is >= 5.0 to be used above UE bundled version (mono 4 does not appear to support .NET 4.6.2 projects)

Only use msbuild in mono when using an installed copy and not just when >= 5.0. The UE dist uses 5.16 but does not have the msbuild executable or dependencies (will address for 4.24 - UE-78476)

Explicitly use bash when calling .RunUAT.sh from from root stub file for OSs that are set to use a different shell

diff --git a/Engine/Build/BatchFiles/Linux/SetupMono.sh b/Engine/Build/BatchFiles/Linux/SetupMono.sh
--- a/Engine/Build/BatchFiles/Linux/SetupMono.sh
+++ b/Engine/Build/BatchFiles/Linux/SetupMono.sh
@@ -13,12 +13,12 @@ bash FixDependencyFiles.sh
 IS_MONO_INSTALLED=0
 MONO_VERSION_PATH=$(command -v mono) || true
 if [ "$UE_USE_SYSTEM_MONO" == "1" ] && [ ! $MONO_VERSION_PATH == "" ] && [ -f $MONO_VERSION_PATH ]; then
-	# If Mono is installed, check if it's 4.0.2 or higher
+	# If Mono is installed, check if it's 5.0 or higher
 	MONO_VERSION_PREFIX="Mono JIT compiler version "
 	MONO_VERSION_PREFIX_LEN=${#MONO_VERSION_PREFIX}
 	MONO_VERSION=`"${MONO_VERSION_PATH}" --version |grep "$MONO_VERSION_PREFIX"`
 	MONO_VERSION=(`echo ${MONO_VERSION:MONO_VERSION_PREFIX_LEN} |tr '.' ' '`)
-	if [ ${MONO_VERSION[0]} -ge 4 ]; then
+	if [ ${MONO_VERSION[0]} -ge 5 ]; then
 		if [ ${MONO_VERSION[1]} -ge 0 ] || [ ${MONO_VERSION[2]} -ge 2 ]; then
 			IS_MONO_INSTALLED=1
 		fi
diff --git a/Engine/Build/BatchFiles/Mac/SetupMono.sh b/Engine/Build/BatchFiles/Mac/SetupMono.sh
--- a/Engine/Build/BatchFiles/Mac/SetupMono.sh
+++ b/Engine/Build/BatchFiles/Mac/SetupMono.sh
@@ -5,7 +5,7 @@ sh FixMonoFiles.sh
 sh FixDependencyFiles.sh
 
 IS_MONO_INSTALLED=0
-IS_MONO_5_INSTALLED=0
+IS_MS_BUILD_AVAILABLE=0
 MONO_VERSION_PATH=`which mono` || true
 
 # if we can't find mono path, try one last hail mary of a standard install location
@@ -19,7 +19,7 @@ if [ "$MONO_VERSION_PATH" == "" ] || [ ! -f $MONO_VERSION_PATH ]; then
 fi
 
 if [ ! $MONO_VERSION_PATH == "" ] && [ -f $MONO_VERSION_PATH ]; then
-	# If Mono is installed, check if it's 4.0.2 or higher
+	# If Mono is installed, check if it's 5.0 or higher
 	MONO_VERSION_PREFIX="Mono JIT compiler version "
 	MONO_VERSION_PREFIX_LEN=${#MONO_VERSION_PREFIX}
 	MONO_VERSION=`"${MONO_VERSION_PATH}" --version |grep "$MONO_VERSION_PREFIX"`
@@ -27,17 +27,11 @@ if [ ! $MONO_VERSION_PATH == "" ] && [ -f $MONO_VERSION_PATH ]; then
 	MONO_VERSION=(`echo ${MONO_VERSION:MONO_VERSION_PREFIX_LEN} |tr '.' ' '`)
 	if [ ${MONO_VERSION[0]} -ge 5 ]; then # Allow any Mono 5.x and up
 		IS_MONO_INSTALLED=1
-		IS_MONO_5_INSTALLED=1
-	elif [ ${MONO_VERSION[0]} -eq 4 ]; then
-		if [ ${MONO_VERSION[1]} -eq 0 ] && [ ${MONO_VERSION[2]} -ge 2 ]; then
-			IS_MONO_INSTALLED=1
-		elif [ ${MONO_VERSION[1]} -gt 0 ] && [ ${MONO_VERSION[1]} -lt 6 ]; then # Mono 4.6 has issues on macOS 10.12
-			IS_MONO_INSTALLED=1
-		fi
+		IS_MS_BUILD_AVAILABLE=0
 	fi
 fi
 
-# Setup bundled Mono if cannot use installed one
+# Setup bundled Mono if cannot use installed one. Note this is 5.16 but does not currently have msbuild bundled
 if [ $IS_MONO_INSTALLED -eq 0 ]; then
 	echo Setting up Mono
 	CUR_DIR=`pwd`
@@ -47,7 +41,7 @@ if [ $IS_MONO_INSTALLED -eq 0 ]; then
 	export LD_LIBRARY_PATH=$UE_MONO_DIR/lib:$LD_LIBRARY_PATH
 else
 	export IS_MONO_INSTALLED=$IS_MONO_INSTALLED
-	export IS_MONO_5_INSTALLED=$IS_MONO_5_INSTALLED
+	export IS_MS_BUILD_AVAILABLE=$IS_MS_BUILD_AVAILABLE
 fi
 
 cd "$START_DIR"
diff --git a/Engine/Build/BatchFiles/RunUAT.sh b/Engine/Build/BatchFiles/RunUAT.sh
--- a/Engine/Build/BatchFiles/RunUAT.sh
+++ b/Engine/Build/BatchFiles/RunUAT.sh
@@ -77,7 +77,7 @@ if [ "$UATCompileArg" = "-compile" ]; then
 		"${SCRIPT_DIR}/FindPlatformExtensionSources.sh"
 
 		# mono 5.0 and up include msbuild
-		if [ "$IS_MONO_5_INSTALLED" == "1" ]; then
+		if [ "$IS_MS_BUILD_AVAILABLE" == "1" ]; then
 			BUILD_TOOL=msbuild
 		else
 			BUILD_TOOL=xbuild
