Fix error generating project files on Mac.

#rnx

diff --git a/Engine/Build/BatchFiles/FindPlatformExtensionSources.sh b/Engine/Build/BatchFiles/FindPlatformExtensionSources.sh
--- a/Engine/Build/BatchFiles/FindPlatformExtensionSources.sh
+++ b/Engine/Build/BatchFiles/FindPlatformExtensionSources.sh
@@ -22,7 +22,7 @@ for PLATFORM in ../../Engine/Platforms/*; do
 	for SOURCE_FILE in ${PLATFORM}/Source/Programs/UnrealBuildTool/*; do
 		echo "<PlatformExtensionCompile Include=\"..\\..\\${SOURCE_FILE//\//\\}\">" >> $REFERENCE_FILE
 		echo "    <Link>Platform\${PLATFORM##*/}\${SOURCE_FILE##*/}</Link>" >> $REFERENCE_FILE
-		echo "</PlatformExtensionCompile>" >> %REFERENCE_FILE%
+		echo "</PlatformExtensionCompile>" >> $REFERENCE_FILE
 	done
 done
 
