Add Changelist parameter to SyncBuildRunLatest mac script

diff --git a/Engine/Build/BatchFiles/Mac/SyncBuildRunLatest.command b/Engine/Build/BatchFiles/Mac/SyncBuildRunLatest.command
--- a/Engine/Build/BatchFiles/Mac/SyncBuildRunLatest.command
+++ b/Engine/Build/BatchFiles/Mac/SyncBuildRunLatest.command
@@ -1,7 +1,8 @@
 #!/bin/sh
 
 #script to automatically sync and run most recent change in this stream in specified project
-#	parameters: .uproject to sync/build/run
+#	parameter 1: .uproject to sync/build/run
+#	parameter 2: (optional) manually specified changelist 
 #	example usage 1 (relative path from root stream directory): ./SyncBuildRunLatest.command QAGame/QAGame.uproject
 #	example usage 2 (absolute disk path) : ./SyncBuildRunLatest.command \Users\your.user\UE4\Dev-Main\QAGame\QAGame.uproject
 #
@@ -33,6 +34,7 @@
 
 
 project=$1
+changelist=$2
 
 if pgrep -x "UE4Editor" > /dev/null
 then
@@ -73,14 +75,19 @@ p4 sync -f "$ue4_root_directory_server"Engine/Source/Programs/DotNETCommon/MetaD
 path_to_sync=$ue4_root_directory_server...
 echo Synching $path_to_sync
 
-p4 changes -m 1 -s submitted "//UE4/Dev-VR/..."
-
-latest_change_string=$(p4 changes -m 1 -s submitted "$path_to_sync")
-echo $latest_change_string
-arr=($latest_change_string)
-echo ${arr[1]}
-
-changelist=${arr[1]}
+if [ -z "$changelist" ]; then
+    echo "changelist unset, getting latest"
+	p4 changes -m 1 -s submitted "//UE4/Dev-VR/..."
+	
+	latest_change_string=$(p4 changes -m 1 -s submitted "$path_to_sync")
+	echo $latest_change_string
+	arr=($latest_change_string)
+	echo ${arr[1]}
+	
+	changelist=${arr[1]}
+else
+	echo "Manually specifying changelist";
+fi
 
 echo changelist="$changelist"
 
