Added support for AoT compilation for mono UnrealBuildTool (~ %10 savings on a noop build due to no JIT)
- Must be opted into using "defaults write com.epicgames.ue4 MonoAOT 1"

diff --git a/Engine/Build/BatchFiles/Mac/GenerateProjectFiles.sh b/Engine/Build/BatchFiles/Mac/GenerateProjectFiles.sh
--- a/Engine/Build/BatchFiles/Mac/GenerateProjectFiles.sh
+++ b/Engine/Build/BatchFiles/Mac/GenerateProjectFiles.sh
@@ -25,5 +25,13 @@ if [ -f "$BASE_PATH/../../../Source/Programs/UnrealBuildTool/UnrealBuildTool.csp
 	xbuild "$BASE_PATH/../../../Source/Programs/UnrealBuildTool/UnrealBuildTool.csproj" /property:Configuration="Development" /verbosity:quiet /nologo |grep -i error
 fi
 
+WANT_AOT="`defaults read com.epicgames.ue4 MonoAOT`"
+if [ $WANT_AOT == "1" ]; then
+	if [ ! -f "$BASE_PATH/../../../Binaries/DotNET/UnrealBuildTool.exe.dylib" ]; then
+		echo Compiling UnrealBuildTool to native...
+		mono --aot "$BASE_PATH/../../../Binaries/DotNET/UnrealBuildTool.exe"
+	fi
+fi
+
 # pass all parameters to UBT
 mono "$BASE_PATH/../../../Binaries/DotNET/UnrealBuildTool.exe" -projectfiles "$@"
