Some small improvements to the SyncBuildRunLatest command for UnrealGameSync-like functionality on Mac.

diff --git a/Engine/Build/BatchFiles/Mac/SyncBuildRunLatest.command b/Engine/Build/BatchFiles/Mac/SyncBuildRunLatest.command
--- a/Engine/Build/BatchFiles/Mac/SyncBuildRunLatest.command
+++ b/Engine/Build/BatchFiles/Mac/SyncBuildRunLatest.command
@@ -39,7 +39,7 @@ changelist=$2
 if pgrep -x "UE4Editor" > /dev/null
 then
     echo "UE4Editor already running, please close all UE4Editor processes..."
-    killall -9 "UE4Editor"
+    #killall -9 "UE4Editor"
     exit 1
 else
     echo "No UE4Editor processes running."
@@ -54,6 +54,14 @@ fi
 
 #echo $(p4 where ../../..)
 
+if [ $? -eq 0 ]
+then
+  echo "p4 where successful"
+else
+  echo "Error running p4 where" >&2
+  exit 1
+fi
+
 #convert local current directory path to perforce server path
 ue4_root_directory_server=$(p4 where ../../..)
 arr=($ue4_root_directory_server)
@@ -70,14 +78,38 @@ echo $ue4_root_directory_server
 
 #have to force sync these two files after running SyncProject previously
 p4 sync -f "$ue4_root_directory_server"Engine/Build/Build.version
+
+if [ $? -eq 0 ]
+then
+  echo "p4 sync successful"
+else
+  echo "Error running p4 sync" >&2
+  exit 1
+fi
+
 p4 sync -f "$ue4_root_directory_server"Engine/Source/Programs/DotNETCommon/MetaData.cs
 
+if [ $? -eq 0 ]
+then
+  echo "p4 sync successful"
+else
+  echo "Error running p4 sync" >&2
+  exit 1
+fi
+
 path_to_sync=$ue4_root_directory_server...
 echo Synching $path_to_sync
 
 if [ -z "$changelist" ]; then
     echo "changelist unset, getting latest"
 	p4 changes -m 1 -s submitted "//UE4/Dev-VR/..."
+	if [ $? -eq 0 ]
+	then
+	  echo "p4 changes successful"
+	else
+	  echo "Error running p4 changes" >&2
+	  exit 1
+	fi
 	
 	latest_change_string=$(p4 changes -m 1 -s submitted "$path_to_sync")
 	echo $latest_change_string
@@ -90,6 +122,7 @@ else
 fi
 
 echo changelist="$changelist"
+echo $latest_change_string
 
 if pgrep -x "UE4Editor" > /dev/null
 then
@@ -101,8 +134,17 @@ else
 fi
 
 echo $ue4_root_directory_local
+echo "../RunUAT.command SyncProject -CL=$changelist -Project="$project" -build -run"
 
 ../RunUAT.command SyncProject -CL=$changelist -Project="$project" -build -run
+if [ $? -eq 0 ]
+then
+  echo "RunUAT.command SyncProject successful"
+else
+  echo "Error running RunUAT.command SyncProject -build -run, aborting..." >&2
+  echo "(See above for errors)" >&2
+  exit 1
+fi
 
 open -a "$ue4_root_directory_local"Engine/Binaries/Mac/UE4Editor.app --args "$project"
 
