Update to UGS-like script for mac.  Fail if specified changelist is not able to be found on the server

diff --git a/Engine/Build/BatchFiles/Mac/SyncBuildRunLatest.command b/Engine/Build/BatchFiles/Mac/SyncBuildRunLatest.command
--- a/Engine/Build/BatchFiles/Mac/SyncBuildRunLatest.command
+++ b/Engine/Build/BatchFiles/Mac/SyncBuildRunLatest.command
@@ -102,7 +102,7 @@ echo Synching $path_to_sync
 
 if [ -z "$changelist" ]; then
     echo "changelist unset, getting latest"
-	p4 changes -m 1 -s submitted "//UE4/Dev-VR/..."
+	p4 changes -m 1 -s submitted "$path_to_sync"
 	if [ $? -eq 0 ]
 	then
 	  echo "p4 changes successful"
@@ -119,6 +119,17 @@ if [ -z "$changelist" ]; then
 	changelist=${arr[1]}
 else
 	echo "Manually specifying changelist";
+	latest_change_string=$(p4 describe $changelist)
+	#echo $latest_change_string
+	arr=($latest_change_string)
+	#echo ${arr[1]}
+	if [ ${arr[1]} -eq $changelist ]; 
+	then
+		echo found changelist ${arr[1]} on Perforce server
+	else
+		echo "Failed to locate specified changelist"
+		exit 1
+	fi
 fi
 
 echo changelist="$changelist"
