Improved the AoT compilation of DotNET executables on Mac

diff --git a/Engine/Build/BatchFiles/Mac/GenerateProjectFiles.sh b/Engine/Build/BatchFiles/Mac/GenerateProjectFiles.sh
--- a/Engine/Build/BatchFiles/Mac/GenerateProjectFiles.sh
+++ b/Engine/Build/BatchFiles/Mac/GenerateProjectFiles.sh
@@ -26,11 +26,41 @@ if [ -f "$BASE_PATH/../../../Source/Programs/UnrealBuildTool/UnrealBuildTool.csp
 fi
 
 WANT_AOT="`defaults read com.epicgames.ue4 MonoAOT`"
+OPENSSL="./../../../Binaries/DotNET/IOS/openssl.exe"
 if [ $WANT_AOT == "1" ]; then
-	if [ ! -f "$BASE_PATH/../../../Binaries/DotNET/UnrealBuildTool.exe.dylib" ]; then
-		echo Compiling UnrealBuildTool to native...
-		mono --aot "$BASE_PATH/../../../Binaries/DotNET/UnrealBuildTool.exe"
-	fi
+        for i in $BASE_PATH/../../../Binaries/DotNET/*.dll;
+        do
+            if [ ! -f "$i.dylib" ]; then
+        		echo Compiling $i to native...
+                mono --aot $i > /dev/null 2>&1;
+            fi
+        done
+
+        for i in $BASE_PATH/../../../Binaries/DotNET/*.exe;
+        do
+            if [ ! -f "$i.dylib" ]; then
+        		echo Compiling $i to native...
+                mono --aot $i > /dev/null 2>&1;
+            fi
+        done
+
+        for i in $BASE_PATH/../../../Binaries/DotNET/IOS/*.dll;
+        do
+            if [ ! -f "$i.dylib" ]; then
+        		echo Compiling $i to native...
+                mono --aot $i > /dev/null 2>&1;
+            fi
+        done
+
+        for i in $BASE_PATH/../../../Binaries/DotNET/IOS/*.exe;
+        do
+            if [ ! -f "$i.dylib" ]; then
+                if [ $i != $OPENSSL ]; then
+                    echo Compiling $i to native...
+                    mono --aot $i > /dev/null 2>&1;
+                fi
+            fi
+        done
 fi
 
 # pass all parameters to UBT
