Changed the AoT processing to check timestamps between exe/dll and .dylib

diff --git a/Engine/Build/BatchFiles/Mac/GenerateProjectFiles.sh b/Engine/Build/BatchFiles/Mac/GenerateProjectFiles.sh
--- a/Engine/Build/BatchFiles/Mac/GenerateProjectFiles.sh
+++ b/Engine/Build/BatchFiles/Mac/GenerateProjectFiles.sh
@@ -30,7 +30,7 @@ OPENSSL="./../../../Binaries/DotNET/IOS/openssl.exe"
 if [ $WANT_AOT == "1" ]; then
         for i in $BASE_PATH/../../../Binaries/DotNET/*.dll;
         do
-            if [ ! -f "$i.dylib" ]; then
+            if test "$i" -nt "$i.dylib"; then
         		echo Compiling $i to native...
                 mono --aot $i > /dev/null 2>&1;
             fi
@@ -38,7 +38,7 @@ if [ $WANT_AOT == "1" ]; then
 
         for i in $BASE_PATH/../../../Binaries/DotNET/*.exe;
         do
-            if [ ! -f "$i.dylib" ]; then
+            if test "$i" -nt "$i.dylib"; then
         		echo Compiling $i to native...
                 mono --aot $i > /dev/null 2>&1;
             fi
@@ -46,7 +46,7 @@ if [ $WANT_AOT == "1" ]; then
 
         for i in $BASE_PATH/../../../Binaries/DotNET/IOS/*.dll;
         do
-            if [ ! -f "$i.dylib" ]; then
+            if test "$i" -nt "$i.dylib"; then
         		echo Compiling $i to native...
                 mono --aot $i > /dev/null 2>&1;
             fi
@@ -54,7 +54,7 @@ if [ $WANT_AOT == "1" ]; then
 
         for i in $BASE_PATH/../../../Binaries/DotNET/IOS/*.exe;
         do
-            if [ ! -f "$i.dylib" ]; then
+            if test "$i" -nt "$i.dylib"; then
                 if [ $i != $OPENSSL ]; then
                     echo Compiling $i to native...
                     mono --aot $i > /dev/null 2>&1;
