Added the __MonoCS__ define to the Mac specific GenerateProjectFiles.sh file

diff --git a/Engine/Build/BatchFiles/Mac/GenerateProjectFiles.sh b/Engine/Build/BatchFiles/Mac/GenerateProjectFiles.sh
--- a/Engine/Build/BatchFiles/Mac/GenerateProjectFiles.sh
+++ b/Engine/Build/BatchFiles/Mac/GenerateProjectFiles.sh
@@ -22,7 +22,7 @@ fi
 . "$BASE_PATH/SetupMono.sh" "$BASE_PATH"
 
 if [ -f "$BASE_PATH/../../../Source/Programs/UnrealBuildTool/UnrealBuildTool.csproj" ]; then
-	xbuild "$BASE_PATH/../../../Source/Programs/UnrealBuildTool/UnrealBuildTool.csproj" /property:Configuration="Development" /verbosity:quiet /nologo |grep -i error
+	xbuild "$BASE_PATH/../../../Source/Programs/UnrealBuildTool/UnrealBuildTool.csproj" /p:DefineConstants=__MonoCS__ /property:Configuration="Development" /verbosity:quiet /nologo |grep -i error
 fi
 
 WANT_AOT="`defaults read com.epicgames.ue4 MonoAOT`"
