Addresses issue where mono > 5 is installed but msbuild isn't

#tests tested with mono installed and with renaming the installed mono dir to verify things were detected properly and msbuild / xbuild selected

diff --git a/Engine/Build/BatchFiles/Mac/SetupMono.sh b/Engine/Build/BatchFiles/Mac/SetupMono.sh
--- a/Engine/Build/BatchFiles/Mac/SetupMono.sh
+++ b/Engine/Build/BatchFiles/Mac/SetupMono.sh
@@ -18,27 +18,38 @@ if [ "$MONO_VERSION_PATH" == "" ] || [ ! -f $MONO_VERSION_PATH ]; then
 	fi
 fi
 
+
+MONO_VERSION_PREFIX="Mono JIT compiler version "
+MONO_VERSION_PREFIX_LEN=${#MONO_VERSION_PREFIX}
+
 if [ ! $MONO_VERSION_PATH == "" ] && [ -f $MONO_VERSION_PATH ]; then
-	# If Mono is installed, check if it's 5.0 or higher
-	MONO_VERSION_PREFIX="Mono JIT compiler version "
-	MONO_VERSION_PREFIX_LEN=${#MONO_VERSION_PREFIX}
+	# Get the version string and major version of the mono in our path
 	MONO_VERSION=`"${MONO_VERSION_PATH}" --version |grep "$MONO_VERSION_PREFIX"`
-	echo "Running installed mono, version: " ${MONO_VERSION}
-	MONO_VERSION=(`echo ${MONO_VERSION:MONO_VERSION_PREFIX_LEN} |tr '.' ' '`)
-	if [ ${MONO_VERSION[0]} -ge 5 ]; then # Allow any Mono 5.x and up
+	MONO_MAJOR_VERSION=(`echo ${MONO_VERSION:MONO_VERSION_PREFIX_LEN} |tr '.' ' '`)
+
+	if [ ${MONO_MAJOR_VERSION[0]} -ge 5 ]; then # Allow any Mono 5.x and up
 		IS_MONO_INSTALLED=1
-		IS_MS_BUILD_AVAILABLE=1
+		# see if msbuild is installed
+		MS_BUILD_PATH=`which msbuild` || true			
+		if [ -f "$MS_BUILD_PATH" ]; then
+			IS_MS_BUILD_AVAILABLE=1
+			echo "Running system mono/msbuild, version: ${MONO_VERSION}"
+		else		
+			echo "Running system mono, version: ${MONO_VERSION}"
+		fi
 	fi
 fi
 
 # Setup bundled Mono if cannot use installed one. Note this is 5.16 but does not currently have msbuild bundled
 if [ $IS_MONO_INSTALLED -eq 0 ]; then
-	echo Setting up Mono
 	CUR_DIR=`pwd`
 	export UE_MONO_DIR=$CUR_DIR/../../../Binaries/ThirdParty/Mono/Mac
 	export PATH=$UE_MONO_DIR/bin:$PATH
 	export MONO_PATH=$UE_MONO_DIR/lib:$MONO_PATH
 	export LD_LIBRARY_PATH=$UE_MONO_DIR/lib:$LD_LIBRARY_PATH
+	MONO_VERSION_PATH=`which mono` || true
+	MONO_VERSION=`"${MONO_VERSION_PATH}" --version |grep "$MONO_VERSION_PREFIX"`
+	echo "Running bundled mono, version: ${MONO_VERSION}"
 else
 	export IS_MONO_INSTALLED=$IS_MONO_INSTALLED
 	export IS_MS_BUILD_AVAILABLE=$IS_MS_BUILD_AVAILABLE
