Adding -w to the grep (now -wi) in the bash script that builds UBT so it isn't triggered by warnings that involved a variable with 'Error' in its name.

(In my case it was 'Missing XML comment for publicly visible type or member 'ReadOnlyTargetRules.bUndefinedIdentifierErrors').

diff --git a/Engine/Build/BatchFiles/Linux/RunXBuild.sh b/Engine/Build/BatchFiles/Linux/RunXBuild.sh
--- a/Engine/Build/BatchFiles/Linux/RunXBuild.sh
+++ b/Engine/Build/BatchFiles/Linux/RunXBuild.sh
@@ -20,7 +20,7 @@ if [ ! -f Build/BatchFiles/Linux/RunXBuild.sh ]; then
 	exit 1
 fi
 
-xbuild /verbosity:quiet /nologo "$@" |grep -i error
+xbuild /verbosity:quiet /nologo "$@" |grep -wi error
 if [ $? -ne 1 ]; then
 	exit 1
 else
diff --git a/Engine/Build/BatchFiles/Mac/RunXBuild.sh b/Engine/Build/BatchFiles/Mac/RunXBuild.sh
--- a/Engine/Build/BatchFiles/Mac/RunXBuild.sh
+++ b/Engine/Build/BatchFiles/Mac/RunXBuild.sh
@@ -20,7 +20,7 @@ if [ ! -f Build/BatchFiles/Mac/RunXBuild.sh ]; then
 	exit 1
 fi
 
-xbuild /verbosity:quiet /nologo "$@" |grep -i error
+xbuild /verbosity:quiet /nologo "$@" |grep -wi error
 if [ $? -ne 1 ]; then
 	exit 1
 else
