Use msbuild in mono for Linux when using an installed copy and mono >= 5.0.

Pulled code from Andrew's CL 7828885 for the Mac.

diff --git a/Engine/Build/BatchFiles/Linux/SetupMono.sh b/Engine/Build/BatchFiles/Linux/SetupMono.sh
--- a/Engine/Build/BatchFiles/Linux/SetupMono.sh
+++ b/Engine/Build/BatchFiles/Linux/SetupMono.sh
@@ -9,9 +9,10 @@ export HOST_ARCH=x86_64-unknown-linux-gnu
 
 bash FixMonoFiles.sh
 bash FixDependencyFiles.sh
-
 IS_MONO_INSTALLED=0
+IS_MS_BUILD_AVAILABLE=0
 MONO_VERSION_PATH=$(command -v mono) || true
+
 if [ "$UE_USE_SYSTEM_MONO" == "1" ] && [ ! $MONO_VERSION_PATH == "" ] && [ -f $MONO_VERSION_PATH ]; then
 	# If Mono is installed, check if it's 5.0 or higher
 	MONO_VERSION_PREFIX="Mono JIT compiler version "
@@ -21,6 +22,15 @@ if [ "$UE_USE_SYSTEM_MONO" == "1" ] && [ ! $MONO_VERSION_PATH == "" ] && [ -f $M
 	if [ ${MONO_VERSION[0]} -ge 5 ]; then
 		if [ ${MONO_VERSION[1]} -ge 0 ] || [ ${MONO_VERSION[2]} -ge 2 ]; then
 			IS_MONO_INSTALLED=1
+
+			# see if msbuild is installed
+			MS_BUILD_PATH=`which msbuild` || true
+			if [ -f "$MS_BUILD_PATH" ]; then
+				IS_MS_BUILD_AVAILABLE=1
+				echo "Running system mono/msbuild, version: ${MONO_VERSION}"
+			else
+				echo "Running system mono, version: ${MONO_VERSION}"
+			fi
 		fi
 	fi
 fi
@@ -34,6 +44,9 @@ if [ $IS_MONO_INSTALLED -eq 0 ]; then
 	export MONO_PATH=$UE_MONO_DIR/lib/mono/4.5:$MONO_PATH
 	export MONO_CFG_DIR=$UE_MONO_DIR/etc
 	export LD_LIBRARY_PATH=$UE_MONO_DIR/$HOST_ARCH/lib:$LD_LIBRARY_PATH
+else
+	export IS_MONO_INSTALLED=$IS_MONO_INSTALLED
+	export IS_MS_BUILD_AVAILABLE=$IS_MS_BUILD_AVAILABLE
 fi
 
 cd "$START_DIR"
diff --git a/Engine/Source/Programs/AutomationTool/AutomationUtils/LinuxHostPlatform.cs b/Engine/Source/Programs/AutomationTool/AutomationUtils/LinuxHostPlatform.cs
--- a/Engine/Source/Programs/AutomationTool/AutomationUtils/LinuxHostPlatform.cs
+++ b/Engine/Source/Programs/AutomationTool/AutomationUtils/LinuxHostPlatform.cs
@@ -7,14 +7,35 @@ using System.Text;
 using System.Diagnostics;
 using System.IO;
 using UnrealBuildTool;
+using Tools.DotNETCommon;
 
 namespace AutomationTool
 {
 	class LinuxHostPlatform : HostPlatform
 	{
+		static string CachedMsBuildTool = "";
+
 		public override string GetMsBuildExe()
 		{
-			return "xbuild";
+			// As of 5.0 mono comes with msbuild which performs better. If that's installed then use it
+			if (string.IsNullOrEmpty(CachedMsBuildTool))
+			{
+				int Value;
+				bool CanUseMsBuild = (int.TryParse(Environment.GetEnvironmentVariable("UE_USE_SYSTEM_MONO"), out Value) &&
+						Value != 0 &&
+						!string.IsNullOrEmpty(CommandUtils.WhichApp("msbuild")));
+
+				if (CanUseMsBuild)
+				{
+					CachedMsBuildTool = "msbuild";
+				}
+				else
+				{
+					CachedMsBuildTool = "xbuild";
+				}
+			}
+
+			return CachedMsBuildTool;
 		}
 
 		public override string RelativeBinariesFolder
