Fix check for UAT running from installed build.

diff --git a/Engine/Build/BatchFiles/RunUAT.sh b/Engine/Build/BatchFiles/RunUAT.sh
--- a/Engine/Build/BatchFiles/RunUAT.sh
+++ b/Engine/Build/BatchFiles/RunUAT.sh
@@ -51,7 +51,7 @@ else
 	UATCompileArg=-compile
 fi
 
-if [ -f Build\InstalledBuild.txt ]; then
+if [ -f Build/InstalledBuild.txt ]; then
 	UATCompileArg=
 fi
 
