Move restricted UBT source files into the Engine/Restricted folder.

diff --git a/Engine/Build/BatchFiles/GenerateProjectFiles.bat b/Engine/Build/BatchFiles/GenerateProjectFiles.bat
--- a/Engine/Build/BatchFiles/GenerateProjectFiles.bat
+++ b/Engine/Build/BatchFiles/GenerateProjectFiles.bat
@@ -28,7 +28,7 @@ rem ## Get the path to MSBuild
 call "%~dp0GetMSBuildPath.bat"
 if errorlevel 1 goto Error_NoVisualStudioEnvironment
 
-call "%~dp0FixDependencyFiles.bat"
+rem call "%~dp0FixDependencyFiles.bat"
 
 rem ## If we're using VS2017, check that NuGet package manager is installed. MSBuild fails to compile C# projects from the command line with a cryptic error if it's not: 
 rem ## https://developercommunity.visualstudio.com/content/problem/137779/the-getreferencenearesttargetframeworktask-task-wa.html
@@ -60,7 +60,19 @@ if not [%MSBUILD_EXE%] == [%MSBUILD_15_EXE_WITH_NUGET%] goto Error_RequireNugetP
 rem Check to see if the files in the UBT directory have changed. We conditionally include platform files from the .csproj file, but MSBuild doesn't recognize the dependency when new files are added. 
 md ..\Intermediate\Build >nul 2>nul
 dir /s /b Programs\UnrealBuildTool\*.cs >..\Intermediate\Build\UnrealBuildToolFiles.txt
-if exist ..\..\Platforms dir /s /b ..\..\Platforms\*.cs >> ..\Intermediate\Build\UnrealBuildToolFiles.txt
+
+if not exist ..\Platforms goto NoPlatforms
+for /d %%D in (..\Platforms\*) do (
+	if exist %%D\Source\Programs\UnrealBuildTool dir /s /b %%D\Source\Programs\UnrealBuildTool\*.cs >> ..\Intermediate\Build\UnrealBuildToolFiles.txt
+)
+:NoPlatforms
+
+if not exist ..\Restricted goto NoRestricted
+for /d %%D in (..\Restricted\*) do (
+	if exist %%D\Source\Programs\UnrealBuildTool dir /s /b %%D\Source\Programs\UnrealBuildTool\*.cs >> ..\Intermediate\Build\UnrealBuildToolFiles.txt
+)
+:NoRestricted
+
 fc /b ..\Intermediate\Build\UnrealBuildToolFiles.txt ..\Intermediate\Build\UnrealBuildToolPrevFiles.txt >nul 2>nul
 if not errorlevel 1 goto SkipClean
 copy /y ..\Intermediate\Build\UnrealBuildToolFiles.txt ..\Intermediate\Build\UnrealBuildToolPrevFiles.txt >nul
diff --git a/Engine/Source/Programs/UnrealBuildTool/UnrealBuildTool.csproj b/Engine/Source/Programs/UnrealBuildTool/UnrealBuildTool.csproj
--- a/Engine/Source/Programs/UnrealBuildTool/UnrealBuildTool.csproj
+++ b/Engine/Source/Programs/UnrealBuildTool/UnrealBuildTool.csproj
@@ -450,6 +450,12 @@
       Replace() fail then RecursiveDir is still an agreeable value -->
       <Link>Platform/$([System.Text.RegularExpressions.Regex]::Replace(%(Compile.RecursiveDir), [\\/].+$, ``))/_</Link>
     </Compile>
+    <Compile Include="../../../Restricted/*/Source/Programs/UnrealBuildTool/**/*.cs">
+      <!-- Similar logic to above for restricted folders, but just remove the first four parts of the
+	  path (eg. NoRedist/Source/Programs/UnrealBuildTool) and allow the files to be based anywhere in the
+	  project from there -->
+      <Link>$([System.Text.RegularExpressions.Regex]::Replace(%(Compile.RecursiveDir), ^([^\\/]+[\\/]){4}, ``))/_</Link>
+    </Compile>
   </ItemGroup>
   <!-- To modify your build process, add your task inside one of the targets below and uncomment it. 
        Other similar extension points exist, see Microsoft.Common.targets.
