It's UE5 now!

diff --git a/Engine/Build/BatchFiles/GenerateProjectFiles.bat b/Engine/Build/BatchFiles/GenerateProjectFiles.bat
--- a/Engine/Build/BatchFiles/GenerateProjectFiles.bat
+++ b/Engine/Build/BatchFiles/GenerateProjectFiles.bat
@@ -7,7 +7,7 @@ rem ## This script is expecting to exist in the UE4 root directory.  It will not
 rem ## if you copy it to a different location and run it.
 
 setlocal
-echo Setting up Unreal Engine 4 project files...
+echo Setting up Unreal Engine 5 project files...
 
 
 rem ## First, make sure the batch file exists in the folder we expect it to.  This is necessary in order to
