UE4->UE5 cosmetic renaming

diff --git a/Engine/Build/BatchFiles/GenerateProjectFiles.bat b/Engine/Build/BatchFiles/GenerateProjectFiles.bat
--- a/Engine/Build/BatchFiles/GenerateProjectFiles.bat
+++ b/Engine/Build/BatchFiles/GenerateProjectFiles.bat
@@ -1,9 +1,9 @@
 @echo off
 
-rem ## Unreal Engine 4 Visual Studio project setup script
+rem ## Unreal Engine 5 Visual Studio project setup script
 rem ## Copyright Epic Games, Inc. All Rights Reserved.
 rem ##
-rem ## This script is expecting to exist in the UE4 root directory.  It will not work correctly
+rem ## This script is expecting to exist in the UE5 root directory.  It will not work correctly
 rem ## if you copy it to a different location and run it.
 
 setlocal
@@ -113,7 +113,7 @@ goto Exit
 
 :Error_RequireNugetPackageManager
 echo.
-echo UE4 requires the NuGet Package Manager to be installed to use %MSBUILD_EXE%. Please run the Visual Studio Installer and add it from the individual components list (in the 'Code Tools' category).
+echo UE5 requires the NuGet Package Manager to be installed to use %MSBUILD_EXE%. Please run the Visual Studio Installer and add it from the individual components list (in the 'Code Tools' category).
 echo.
 pause
 goto Exit
