Suppressed File Not Found error when running GenerateProjectFiles caused by there no longer being any source under Restricted/NotForLicensees

diff --git a/Engine/Build/BatchFiles/GenerateProjectFiles.bat b/Engine/Build/BatchFiles/GenerateProjectFiles.bat
--- a/Engine/Build/BatchFiles/GenerateProjectFiles.bat
+++ b/Engine/Build/BatchFiles/GenerateProjectFiles.bat
@@ -61,13 +61,13 @@ dir /s /b Programs\UnrealBuildTool\*.cs >..\Intermediate\Build\UnrealBuildToolFi
 
 if not exist ..\Platforms goto NoPlatforms
 for /d %%D in (..\Platforms\*) do (
-	if exist %%D\Source\Programs\UnrealBuildTool dir /s /b %%D\Source\Programs\UnrealBuildTool\*.cs >> ..\Intermediate\Build\UnrealBuildToolFiles.txt
+	if exist %%D\Source\Programs\UnrealBuildTool dir /s /b %%D\Source\Programs\UnrealBuildTool\*.cs 2> nul >> ..\Intermediate\Build\UnrealBuildToolFiles.txt
 )
 :NoPlatforms
 
 if not exist ..\Restricted goto NoRestricted
 for /d %%D in (..\Restricted\*) do (
-	if exist %%D\Source\Programs\UnrealBuildTool dir /s /b %%D\Source\Programs\UnrealBuildTool\*.cs >> ..\Intermediate\Build\UnrealBuildToolFiles.txt
+	if exist %%D\Source\Programs\UnrealBuildTool dir /s /b %%D\Source\Programs\UnrealBuildTool\*.cs 2> nul >> ..\Intermediate\Build\UnrealBuildToolFiles.txt
 )
 :NoRestricted
 
