Fixed issue with Build.bat and Clean.bat due to whitespace getting added to the environment variables.

diff --git a/Engine/Build/BatchFiles/Clean.bat b/Engine/Build/BatchFiles/Clean.bat
--- a/Engine/Build/BatchFiles/Clean.bat
+++ b/Engine/Build/BatchFiles/Clean.bat
@@ -31,11 +31,11 @@ if errorlevel 1 goto Error_NoVisualStudioEnvironment
 
 rem ## Compile UBT if the project file exists
 
-SET ProjectFile = "Programs\UnrealBuildTool\UnrealBuildTool.csproj"
-set UBTPath = "..\..\Engine\Binaries\DotNET\UnrealBuildTool.exe"
+set ProjectFile="Programs\UnrealBuildTool\UnrealBuildTool.csproj"
+set UBTPath="..\..\Engine\Binaries\DotNET\UnrealBuildTool.exe"
 if "%UE_USE_DOTNET%" == "1" (
-    set ProjectFile = "Programs\UnrealBuildTool\UnrealBuildToolCore.csproj"
-    set UBTPath = "..\..\Engine\Binaries\DotNET\UnrealBuildTool\UnrealBuildTool.exe"
+    set ProjectFile="Programs\UnrealBuildTool\UnrealBuildToolCore.csproj"
+    set UBTPath="..\..\Engine\Binaries\DotNET\UnrealBuildTool\UnrealBuildTool.exe"
 )
 if not exist %ProjectFile% goto NoProjectFile
 if "%UE_USE_DOTNET%" == "1" (
