- Sety an env var in DesktopPlatform so that children processes know it's being run from an unreal process (UE_DesktopUnrealProcess=1)
- Use this in RunUAT.sh to control how the background workaround for child processes not dying is used for mono UAT

diff --git a/Engine/Build/BatchFiles/RunUAT.sh b/Engine/Build/BatchFiles/RunUAT.sh
--- a/Engine/Build/BatchFiles/RunUAT.sh
+++ b/Engine/Build/BatchFiles/RunUAT.sh
@@ -102,14 +102,24 @@ else
 	LogDir="$uebp_LogFolder"
 fi
 
-trap TermHandler SIGTERM SIGINT
-
-# you can't set a dotted env var nicely in sh, but env will run a command with
-# a list of env vars set, including dotted ones
-echo Start UAT: mono AutomationTool.exe "${Args[@]}"
-env uebp_LogFolder="$LogDir" mono AutomationTool.exe "${Args[@]}" $UATCompileArg &
-UATPid=$!
-wait $UATPid
+# if we are running under UE, we need to run this with the term handler (otherwise canceling a UAT job from the editor
+# can leave mono, etc running in the background, which means we need the PID so we 
+# run it in the background
+if [ "$UE_DesktopUnrealProcess" = "1"]
+	# you can't set a dotted env var nicely in sh, but env will run a command with
+	# a list of env vars set, including dotted ones
+	echo Start UAT Non-Interactively: mono AutomationTool.exe "${Args[@]}"
+	trap TermHandler SIGTERM SIGINT
+	env uebp_LogFolder="$LogDir" mono AutomationTool.exe "${Args[@]}" $UATCompileArg &
+	UATPid=$!
+	wait $UATPid
+else
+	# you can't set a dotted env var nicely in sh, but env will run a command with
+	# a list of env vars set, including dotted ones
+	echo Start UAT Interactively: mono AutomationTool.exe "${Args[@]}"
+	env uebp_LogFolder="$LogDir" mono AutomationTool.exe "${Args[@]}" $UATCompileArg
+fi
+
 UATReturn=$?
 
 # @todo: Copy log files to somewhere useful
diff --git a/Engine/Source/Developer/DesktopPlatform/Private/DesktopPlatformModule.cpp b/Engine/Source/Developer/DesktopPlatform/Private/DesktopPlatformModule.cpp
--- a/Engine/Source/Developer/DesktopPlatform/Private/DesktopPlatformModule.cpp
+++ b/Engine/Source/Developer/DesktopPlatform/Private/DesktopPlatformModule.cpp
@@ -9,10 +9,14 @@ DEFINE_LOG_CATEGORY(LogDesktopPlatform);
 void FDesktopPlatformModule::StartupModule()
 {
 	DesktopPlatform = new FDesktopPlatform();
+
+	FPlatformMisc::SetEnvironmentVar(TEXT("UE_DesktopUnrealProcess"), TEXT("1"));
 }
 
 void FDesktopPlatformModule::ShutdownModule()
 {
+	FPlatformMisc::SetEnvironmentVar(TEXT("UE_DesktopUnrealProcess"), TEXT("0"));
+
 	if (DesktopPlatform != NULL)
 	{
 		delete DesktopPlatform;
