Missing ; then for bash

diff --git a/Engine/Build/BatchFiles/RunUAT.sh b/Engine/Build/BatchFiles/RunUAT.sh
--- a/Engine/Build/BatchFiles/RunUAT.sh
+++ b/Engine/Build/BatchFiles/RunUAT.sh
@@ -105,7 +105,7 @@ fi
 # if we are running under UE, we need to run this with the term handler (otherwise canceling a UAT job from the editor
 # can leave mono, etc running in the background, which means we need the PID so we 
 # run it in the background
-if [ "$UE_DesktopUnrealProcess" = "1" ]
+if [ "$UE_DesktopUnrealProcess" = "1" ]; then
 	# you can't set a dotted env var nicely in sh, but env will run a command with
 	# a list of env vars set, including dotted ones
 	echo Start UAT Non-Interactively: mono AutomationTool.exe "${Args[@]}"
