UAT Net-Core now supports running on Linux as well.
Set UE_USE_DOTNET=1 to enable this.

diff --git a/Engine/Build/BatchFiles/Linux/SetupDotnet.sh b/Engine/Build/BatchFiles/Linux/SetupDotnet.sh
--- a/Engine/Build/BatchFiles/Linux/SetupDotnet.sh
+++ b/Engine/Build/BatchFiles/Linux/SetupDotnet.sh
@@ -1,6 +1,9 @@
 #!/bin/bash
 # Copyright Epic Games, Inc. All Rights Reserved.
 
+START_DIR=`pwd`
+cd "$1"
+
 IS_DOTNET_INSTALLED=0
 DOTNET_VERSION_PATH=$(command -v dotnet) || true
 
@@ -29,4 +32,6 @@ if [ $IS_DOTNET_INSTALLED -eq 0 ]; then
     export DOTNET_ROOT=$UE_DOTNET_DIR
 else
 	export IS_DOTNET_INSTALLED=$IS_DOTNET_INSTALLED
-fi
\ No newline at end of file
+fi
+
+cd "$START_DIR"
diff --git a/Engine/Build/BatchFiles/Linux/SetupEnvironment.sh b/Engine/Build/BatchFiles/Linux/SetupEnvironment.sh
--- a/Engine/Build/BatchFiles/Linux/SetupEnvironment.sh
+++ b/Engine/Build/BatchFiles/Linux/SetupEnvironment.sh
@@ -25,7 +25,7 @@ do
             ;;
         -dotnet) 
             # echo "setting up dotnet"
-            source $2/SetupDotnet.sh "$2"
+            source "$2/SetupDotnet.sh" "$2"
             shift
             ;;
         -*) echo "bad argument $1"
diff --git a/Engine/Build/BatchFiles/RunUAT.sh b/Engine/Build/BatchFiles/RunUAT.sh
--- a/Engine/Build/BatchFiles/RunUAT.sh
+++ b/Engine/Build/BatchFiles/RunUAT.sh
@@ -40,7 +40,7 @@ UATCompileArg=-compile
 
 if [ ! -f Build/BatchFiles/RunUAT.sh ]; then
 	echo "RunUAT ERROR: The script does not appear to be located in the "
-  echo "Engine/Build/BatchFiles directory.  This script must be run from within that directory."
+	echo "Engine/Build/BatchFiles directory.  This script must be run from within that directory."
 	exit 1
 fi
 
@@ -55,14 +55,20 @@ if [ -f Build/InstalledBuild.txt ]; then
 	UATCompileArg=
 fi
 
+EnvironmentType=-mono
+if [ ${UE_USE_DOTNET:=0} -ne 0 ]; then
+	EnvironmentType=-dotnet
+	UATDirectory=Binaries/DotNET/AutomationTool
+fi
+
 if [ "$(uname)" = "Darwin" ]; then
-	# Setup Environment and Mono
-	source "$SCRIPT_DIR/Mac/SetupEnvironment.sh" -mono "$SCRIPT_DIR/Mac"
+	# Setup Environment
+	source "$SCRIPT_DIR/Mac/SetupEnvironment.sh" $EnvironmentType "$SCRIPT_DIR/Mac"
 fi
 
 if [ "$(uname)" = "Linux" ]; then
-	# Setup Environment and Mono
-	source "$SCRIPT_DIR/Mac/SetupEnvironment.sh" -mono "$SCRIPT_DIR/Linux"
+	# Setup Environment
+	source "$SCRIPT_DIR/Linux/SetupEnvironment.sh" $EnvironmentType "$SCRIPT_DIR/Linux"
 fi
 
 
@@ -72,22 +78,30 @@ if [ "$UATCompileArg" = "-compile" ]; then
 		echo No project to compile, attempting to use precompiled AutomationTool
 		UATCompileArg=
 	else
-		# mono 5.0 and up include msbuild
-		if [ "$IS_MS_BUILD_AVAILABLE" == "1" ]; then
-			BUILD_TOOL=msbuild
+		if [ ${UE_USE_DOTNET:=0} -ne 0 ]; then
+			dotnet build Source/Programs/AutomationTool/AutomationToolCore.csproj -c Development
+			if [ $? -ne 0 ]; then
+				echo RunUAT ERROR: AutomationTool failed to compile.
+				exit 1
+			fi
 		else
-			BUILD_TOOL=xbuild
-		fi
-
-		ARGS="/p:Configuration=Development /p:Platform=AnyCPU /verbosity:quiet /nologo /p:NoWarn=1591 /property:AutomationToolProjectOnly=true"
-		ARGS="${ARGS} /p:TargetFrameworkProfile="
-
-		echo "$BUILD_TOOL Source/Programs/AutomationTool/AutomationTool.csproj $ARGS"
-		$BUILD_TOOL Source/Programs/AutomationTool/AutomationTool.csproj $ARGS
-		# make sure it succeeded
-		if [ $? -ne 0 ]; then
-			echo RunUAT ERROR: AutomationTool failed to compile.
-			exit 1
+			# mono 5.0 and up include msbuild
+			if [ "$IS_MS_BUILD_AVAILABLE" == "1" ]; then
+				BUILD_TOOL=msbuild
+			else
+				BUILD_TOOL=xbuild
+			fi
+
+			ARGS="/p:Configuration=Development /p:Platform=AnyCPU /verbosity:quiet /nologo /p:NoWarn=1591 /property:AutomationToolProjectOnly=true"
+			ARGS="${ARGS} /p:TargetFrameworkProfile="
+
+			echo "$BUILD_TOOL Source/Programs/AutomationTool/AutomationTool.csproj $ARGS"
+			$BUILD_TOOL Source/Programs/AutomationTool/AutomationTool.csproj $ARGS
+			# make sure it succeeded
+			if [ $? -ne 0 ]; then
+				echo RunUAT ERROR: AutomationTool failed to compile.
+				exit 1
+			fi
 		fi
 	fi
 fi
@@ -97,27 +111,32 @@ fi
 #run UAT
 cd $UATDirectory
 if [ -z "$uebp_LogFolder" ]; then
-  LogDir="$HOME/Library/Logs/Unreal Engine/LocalBuildLogs"
+	LogDir="$HOME/Library/Logs/Unreal Engine/LocalBuildLogs"
 else
 	LogDir="$uebp_LogFolder"
 fi
 
-# if we are running under UE, we need to run this with the term handler (otherwise canceling a UAT job from the editor
-# can leave mono, etc running in the background, which means we need the PID so we 
-# run it in the background
-if [ "$UE_DesktopUnrealProcess" = "1" ]; then
-	# you can't set a dotted env var nicely in sh, but env will run a command with
-	# a list of env vars set, including dotted ones
-	echo Start UAT Non-Interactively: mono AutomationTool.exe "${Args[@]}"
-	trap TermHandler SIGTERM SIGINT
-	env uebp_LogFolder="$LogDir" mono AutomationTool.exe "${Args[@]}" $UATCompileArg &
-	UATPid=$!
-	wait $UATPid
+if [ ${UE_USE_DOTNET:=0} -ne 0 ]; then
+	echo Start UAT: AutomationTool "${Args[@]}"
+	env uebp_LogFolder="$LogDir" ./AutomationTool "${Args[@]}" $UATCompileArg
 else
-	# you can't set a dotted env var nicely in sh, but env will run a command with
-	# a list of env vars set, including dotted ones
-	echo Start UAT Interactively: mono AutomationTool.exe "${Args[@]}"
-	env uebp_LogFolder="$LogDir" mono AutomationTool.exe "${Args[@]}" $UATCompileArg
+	# if we are running under UE, we need to run this with the term handler (otherwise canceling a UAT job from the editor
+	# can leave mono, etc running in the background, which means we need the PID so we 
+	# run it in the background
+	if [ "$UE_DesktopUnrealProcess" = "1" ]; then
+		# you can't set a dotted env var nicely in sh, but env will run a command with
+		# a list of env vars set, including dotted ones
+		echo Start UAT Non-Interactively: mono AutomationTool.exe "${Args[@]}"
+		trap TermHandler SIGTERM SIGINT
+		env uebp_LogFolder="$LogDir" mono AutomationTool.exe "${Args[@]}" $UATCompileArg &
+		UATPid=$!
+		wait $UATPid
+	else
+		# you can't set a dotted env var nicely in sh, but env will run a command with
+		# a list of env vars set, including dotted ones
+		echo Start UAT Interactively: mono AutomationTool.exe "${Args[@]}"
+		env uebp_LogFolder="$LogDir" mono AutomationTool.exe "${Args[@]}" $UATCompileArg
+	fi
 fi
 
 UATReturn=$?
diff --git a/Engine/Source/Programs/AutomationTool/AutomationUtils/Utils.cs b/Engine/Source/Programs/AutomationTool/AutomationUtils/Utils.cs
--- a/Engine/Source/Programs/AutomationTool/AutomationUtils/Utils.cs
+++ b/Engine/Source/Programs/AutomationTool/AutomationUtils/Utils.cs
@@ -724,7 +724,17 @@ namespace AutomationTool
 			var bCreatedMutex = false;
             var EntryAssemblyLocation = Assembly.GetEntryAssembly().GetOriginalLocation();
 			var LocationHash = EntryAssemblyLocation.GetHashCode();
-            var MutexName = "Global/" + Path.GetFileNameWithoutExtension(EntryAssemblyLocation) + "_" + LocationHash.ToString() + "_Mutex";
+
+			string Prefix = "Global/";
+#if NET_CORE
+			// / is a reservered character for mutexes and thus should be omited when not running on windows
+			// as Global/ and Local/ only has meaning there. Will throw a DirectoryNotFoundException otherwise.
+			if (!Utils.IsRunningOnWindows)
+			{
+				Prefix = "";
+			}
+#endif
+			var MutexName = Prefix + Path.GetFileNameWithoutExtension(EntryAssemblyLocation) + "_" + LocationHash.ToString() + "_Mutex";
 			using (Mutex SingleInstanceMutex = new Mutex(true, MutexName, out bCreatedMutex))
 			{
 				IsSoleInstance = bCreatedMutex;
diff --git a/Engine/Source/Programs/UnrealBuildTool/System/Utils.cs b/Engine/Source/Programs/UnrealBuildTool/System/Utils.cs
--- a/Engine/Source/Programs/UnrealBuildTool/System/Utils.cs
+++ b/Engine/Source/Programs/UnrealBuildTool/System/Utils.cs
@@ -25,6 +25,20 @@ namespace UnrealBuildTool
 		/// </summary>
 #if NET_CORE
 		public static readonly bool IsRunningOnMono = !RuntimeInformation.IsOSPlatform(OSPlatform.Windows); // TODO: This check seems to mainly be used to determine if this is running on windows or not. Should be refactored.
+
+		/// <summary>
+		/// Whether we are currently running on Linux.
+		/// </summary>
+		public static readonly bool IsRunningOnLinux = RuntimeInformation.IsOSPlatform(OSPlatform.Linux);
+		/// <summary>
+		/// Whether we are currently running on a MacOS platform.
+		/// </summary>
+		public static readonly bool IsRunningOnMac = RuntimeInformation.IsOSPlatform(OSPlatform.OSX);
+		
+		/// <summary>
+		/// Whether we are currently running a Windows platform.
+		/// </summary>
+		public static readonly bool IsRunningOnWindows = RuntimeInformation.IsOSPlatform(OSPlatform.Windows);
 #else
 		public static readonly bool IsRunningOnMono = Type.GetType("Mono.Runtime") != null;
 #endif
