Updated mac shell scripts for UBT in NetCore, opt in by setting UE_USE_DOTNET=1, note that due to output path differences a lot of tooling will break when enabling it.

diff --git a/Engine/Build/BatchFiles/Mac/Build.sh b/Engine/Build/BatchFiles/Mac/Build.sh
--- a/Engine/Build/BatchFiles/Mac/Build.sh
+++ b/Engine/Build/BatchFiles/Mac/Build.sh
@@ -5,14 +5,25 @@
 cd "`dirname "$0"`/../../../.."
 
 # Setup Environment and Mono
-source Engine/Build/BatchFiles/Mac/SetupEnvironment.sh -mono Engine/Build/BatchFiles/Mac
+if [ ${UE_USE_DOTNET:=0} -ne 0 ]; then
+  source Engine/Build/BatchFiles/Mac/SetupEnvironment.sh -dotnet Engine/Build/BatchFiles/Mac
+else
+  source Engine/Build/BatchFiles/Mac/SetupEnvironment.sh -mono Engine/Build/BatchFiles/Mac
+fi
 
 # Skip UBT and SWC compile step if we're coming in on an SSH connection (ie remote toolchain)
 if [ -z "$SSH_CONNECTION" ]; then
 	# First make sure that the UnrealBuildTool is up-to-date
-	if ! xbuild /property:Configuration=Development /verbosity:quiet /nologo /p:NoWarn=1591 Engine/Source/Programs/UnrealBuildTool/UnrealBuildTool.csproj; then
-	  echo "Failed to build to build tool (UnrealBuildTool)"
-	  exit 1
+	if [ ${UE_USE_DOTNET:=0} -ne 0 ]; then
+		if ! dotnet build Engine/Source/Programs/UnrealBuildTool/UnrealBuildToolCore.csproj -c Development; then
+			echo "Failed to build to build tool (UnrealBuildTool)"
+			exit 1
+		else
+			if ! xbuild /property:Configuration=Development /verbosity:quiet /nologo /p:NoWarn=1591 Engine/Source/Programs/UnrealBuildTool/UnrealBuildTool.csproj; then
+				echo "Failed to build to build tool (UnrealBuildTool)"
+				exit 1
+			fi
+		fi
 	fi
 
 	# build SCW if specified
@@ -25,8 +36,13 @@ if [ -z "$SSH_CONNECTION" ]; then
 	done
 fi
 
-echo Running Engine/Binaries/DotNET/UnrealBuildTool.exe "$@"
-mono Engine/Binaries/DotNET/UnrealBuildTool.exe "$@"
+if [ ${UE_USE_DOTNET:=0} -ne 0 ]; then
+  echo Running Engine/Binaries/DotNET/UnrealBuildTool/UnrealBuildTool "$@"
+  Engine/Binaries/DotNET/UnrealBuildTool/UnrealBuildTool "$@"
+else
+  echo Running Engine/Binaries/DotNET/UnrealBuildTool.exe "$@"
+  mono Engine/Binaries/DotNET/UnrealBuildTool.exe "$@"
+fi
 
 ExitCode=$?
 if [ $ExitCode -eq 254 ] || [ $ExitCode -eq 255 ] || [ $ExitCode -eq 2 ]; then
diff --git a/Engine/Build/BatchFiles/Mac/GenerateProjectFiles.sh b/Engine/Build/BatchFiles/Mac/GenerateProjectFiles.sh
--- a/Engine/Build/BatchFiles/Mac/GenerateProjectFiles.sh
+++ b/Engine/Build/BatchFiles/Mac/GenerateProjectFiles.sh
@@ -2,7 +2,7 @@
 # Copyright Epic Games, Inc. All Rights Reserved.
 
 echo
-echo Setting up Unreal Engine 4 project files...
+echo Setting up Unreal Engine 5 project files...
 echo
 
 # If ran from somewhere other then the script location we'll have the full base path
@@ -20,51 +20,68 @@ if [ ! -d "$BASE_PATH/../../../Source" ]; then
  exit 1
 fi
 
-source "$BASE_PATH/SetupEnvironment.sh" -mono "$BASE_PATH"
+if [ ${UE_USE_DOTNET:=0} -ne 0 ]; then
+  source "$BASE_PATH/SetupEnvironment.sh" -dotnet "$BASE_PATH"
+else
+  source "$BASE_PATH/SetupEnvironment.sh" -mono "$BASE_PATH"
+fi
 
 if [ -f "$BASE_PATH/../../../Source/Programs/UnrealBuildTool/UnrealBuildTool.csproj" ]; then
-	xbuild "$BASE_PATH/../../../Source/Programs/UnrealBuildTool/UnrealBuildTool.csproj" /property:Configuration="Development" /verbosity:quiet /nologo /p:NoWarn=1591 |grep -i error
+  if [ ${UE_USE_DOTNET:=0} -ne 0 ]; then
+    dotnet build $BASE_PATH/../../../Source/Programs/UnrealBuildTool/UnrealBuildToolCore.csproj -c Development
+  else
+    xbuild "$BASE_PATH/../../../Source/Programs/UnrealBuildTool/UnrealBuildTool.csproj" /property:Configuration="Development" /verbosity:quiet /nologo /p:NoWarn=1591
+  fi
+
+  if [ $? -ne 0 ]; then
+    echo GenerateProjectFiles ERROR: Failed to build UnrealBuildTool
+    exit 1
+  fi
 fi
 
 WANT_AOT="`defaults read com.epicgames.ue4 MonoAOT`"
 OPENSSL="./../../../Binaries/DotNET/IOS/openssl.exe"
 if [ ! -z $WANT_AOT ]; then
-	if [ $WANT_AOT == "1" ]; then
-			for i in $BASE_PATH/../../../Binaries/DotNET/*.dll;
-			do
-				if test "$i" -nt "$i.dylib"; then
-					echo Compiling $i to native...
-					mono --aot $i > /dev/null 2>&1;
-				fi
-			done
+	if [ $WANT_AOT == "1"] && [${UE_USE_DOTNET:=0} -eq 0 ]; then
+		for i in $BASE_PATH/../../../Binaries/DotNET/*.dll;
+		do
+			if test "$i" -nt "$i.dylib"; then
+				echo Compiling $i to native...
+				mono --aot $i > /dev/null 2>&1;
+			fi
+		done
 
-			for i in $BASE_PATH/../../../Binaries/DotNET/*.exe;
-			do
-				if test "$i" -nt "$i.dylib"; then
-					echo Compiling $i to native...
-					mono --aot $i > /dev/null 2>&1;
-				fi
-			done
+		for i in $BASE_PATH/../../../Binaries/DotNET/*.exe;
+		do
+			if test "$i" -nt "$i.dylib"; then
+				echo Compiling $i to native...
+				mono --aot $i > /dev/null 2>&1;
+			fi
+		done
+
+		for i in $BASE_PATH/../../../Binaries/DotNET/IOS/*.dll;
+		do
+			if test "$i" -nt "$i.dylib"; then
+				echo Compiling $i to native...
+				mono --aot $i > /dev/null 2>&1;
+			fi
+		done
 
-			for i in $BASE_PATH/../../../Binaries/DotNET/IOS/*.dll;
-			do
-				if test "$i" -nt "$i.dylib"; then
+		for i in $BASE_PATH/../../../Binaries/DotNET/IOS/*.exe;
+		do
+			if test "$i" -nt "$i.dylib"; then
+				if [ $i != $OPENSSL ]; then
 					echo Compiling $i to native...
 					mono --aot $i > /dev/null 2>&1;
 				fi
-			done
-
-			for i in $BASE_PATH/../../../Binaries/DotNET/IOS/*.exe;
-			do
-				if test "$i" -nt "$i.dylib"; then
-					if [ $i != $OPENSSL ]; then
-						echo Compiling $i to native...
-						mono --aot $i > /dev/null 2>&1;
-					fi
-				fi
-			done
+			fi
+		done
 	fi
 fi
 
 # pass all parameters to UBT
-mono "$BASE_PATH/../../../Binaries/DotNET/UnrealBuildTool.exe" -projectfiles "$@"
+if [ ${UE_USE_DOTNET:=0} -ne 0 ]; then
+  "$BASE_PATH/../../../Binaries/DotNET/UnrealBuildTool/UnrealBuildTool" -projectfiles "$@"
+else
+  mono "$BASE_PATH/../../../Binaries/DotNET/UnrealBuildTool.exe" -projectfiles "$@"
+fi
diff --git a/Engine/Build/BatchFiles/Mac/SetupDotnet.sh b/Engine/Build/BatchFiles/Mac/SetupDotnet.sh
new file mode 100644
--- /dev/null
+++ b/Engine/Build/BatchFiles/Mac/SetupDotnet.sh
@@ -0,0 +1,32 @@
+#!/bin/bash
+# Copyright Epic Games, Inc. All Rights Reserved.
+
+IS_DOTNET_INSTALLED=0
+DOTNET_VERSION_PATH=$(command -v dotnet) || true
+
+if [ "$UE_USE_SYSTEM_DOTNET" == "1" ] && [ ! $DOTNET_VERSION_PATH == "" ] && [ -f $DOTNET_VERSION_PATH ]; then
+	# If dotnet is installed, check that it has a new enough version of the SDK
+	DOTNET_SDKS=(`dotnet --list-sdks | grep -P "(\d*)\.(\d*)\..* \[(.*)\]"`)
+	for DOTNET_SDK in $DOTNET_SDKS
+    do
+        if [ ${DOTNET_SDK[0]} -ge 3 ]; then
+            if [ ${DOTNET_SDK[1]} -ge 1 ]; then
+                IS_DOTNET_INSTALLED=1
+            fi
+        fi
+    done
+    if [ $IS_DOTNET_INSTALLED -eq 0 ]; then
+    echo Unable to find installed dotnet sdk of version 3.1 or newer
+    fi
+fi
+
+# Setup bundled Dotnet if cannot use installed one
+if [ $IS_DOTNET_INSTALLED -eq 0 ]; then
+	echo Setting up bundled DotNet SDK
+    CUR_DIR=`pwd`
+	export UE_DOTNET_DIR=$CUR_DIR/../../../Binaries/ThirdParty/DotNet/Mac
+	export PATH=$UE_DOTNET_DIR:$PATH
+    export DOTNET_ROOT=$UE_DOTNET_DIR
+else
+	export IS_DOTNET_INSTALLED=$IS_DOTNET_INSTALLED
+fi
\ No newline at end of file
diff --git a/Engine/Build/BatchFiles/Mac/SetupEnvironment.sh b/Engine/Build/BatchFiles/Mac/SetupEnvironment.sh
--- a/Engine/Build/BatchFiles/Mac/SetupEnvironment.sh
+++ b/Engine/Build/BatchFiles/Mac/SetupEnvironment.sh
@@ -23,6 +23,11 @@ do
             source "$2/SetupMono.sh" "$2"
             shift
             ;;
+        -dotnet) 
+            # echo "setting up dotnet"
+            source $2/SetupDotnet.sh "$2"
+            shift
+            ;;
         -*) echo "bad argument $1"
             ;;
     esac
diff --git a/Engine/Build/BatchFiles/Mac/XcodeBuild.sh b/Engine/Build/BatchFiles/Mac/XcodeBuild.sh
--- a/Engine/Build/BatchFiles/Mac/XcodeBuild.sh
+++ b/Engine/Build/BatchFiles/Mac/XcodeBuild.sh
@@ -4,15 +4,27 @@
 # (and can take the same arguments) but performs some interpretation of arguments that come from Xcode
 # Values for $ACTION: "" = building, "clean" = cleaning
 
-# Setup Environment & Mono
-source Engine/Build/BatchFiles/Mac/SetupEnvironment.sh -mono Engine/Build/BatchFiles/Mac
+# Setup Environment
+if [ ${UE_USE_DOTNET:=0} -ne 0 ]; then
+  source  Engine/Build/BatchFiles/Mac/SetupEnvironment.sh -dotnet Engine/Build/BatchFiles/Mac
+else
+  source Engine/Build/BatchFiles/Mac/SetupEnvironment.sh -mono Engine/Build/BatchFiles/Mac
+fi
+
 
 # First make sure that the UnrealBuildTool is up-to-date
-if ! xbuild /property:Configuration=Development /verbosity:quiet /nologo /p:NoWarn=1591 Engine/Source/Programs/UnrealBuildTool/UnrealBuildTool.csproj; then
-  echo "Failed to build to build tool (UnrealBuildTool)"
-  exit 1
+if [ ${UE_USE_DOTNET:=0} -ne 0 ]; then
+dotnet build Engine/Source/Programs/UnrealBuildTool/UnrealBuildToolCore.csproj -c Development
+else
+xbuild /property:Configuration=Development /verbosity:quiet /nologo /p:NoWarn=1591 Engine/Source/Programs/UnrealBuildTool/UnrealBuildTool.csproj;
 fi
 
+if [ $? -ne 0 ]; then
+echo "Failed to build the build tool (UnrealBuildTool)"
+exit 1
+fi
+
+
 #echo "Raw Args: $*"
 
 case $1 in 
@@ -88,15 +100,24 @@ if [ "$ACTION" == "build" ]; then
 
 	# Build SCW if this is an editor target
 	if [[ "$TARGET" == *"Editor" ]]; then
-		mono Engine/Binaries/DotNET/UnrealBuildTool.exe ShaderCompileWorker Mac Development
+		if [ ${UE_USE_DOTNET:=0} -ne 0 ]; then
+			Engine/Binaries/DotNET/UnrealBuildTool/UnrealBuildTool ShaderCompileWorker Mac Development
+		else
+			mono Engine/Binaries/DotNET/UnrealBuildTool.exe ShaderCompileWorker Mac Development
+		fi
 	fi
 
 elif [ $ACTION == "clean" ]; then
 	AdditionalFlags="-clean"
 fi
 
+if [ ${UE_USE_DOTNET:=0} -ne 0 ]; then
+echo Running Engine/Binaries/DotNET/UnrealBuildTool/UnrealBuildTool $TARGET $PLATFORM $CONFIGURATION "$TRAILINGARGS" $AdditionalFlags
+Engine/Binaries/DotNET/UnrealBuildTool/UnrealBuildTool $TARGET $PLATFORM $CONFIGURATION "$TRAILINGARGS" $AdditionalFlags
+else
 echo Running Engine/Binaries/DotNET/UnrealBuildTool.exe $TARGET $PLATFORM $CONFIGURATION "$TRAILINGARGS" $AdditionalFlags
 mono Engine/Binaries/DotNET/UnrealBuildTool.exe $TARGET $PLATFORM $CONFIGURATION "$TRAILINGARGS" $AdditionalFlags
+fi
 
 ExitCode=$?
 if [ $ExitCode -eq 254 ] || [ $ExitCode -eq 255 ] || [ $ExitCode -eq 2 ]; then
diff --git a/Engine/Source/Programs/UnrealBuildTool/Executors/LocalExecutor.cs b/Engine/Source/Programs/UnrealBuildTool/Executors/LocalExecutor.cs
--- a/Engine/Source/Programs/UnrealBuildTool/Executors/LocalExecutor.cs
+++ b/Engine/Source/Programs/UnrealBuildTool/Executors/LocalExecutor.cs
@@ -71,11 +71,18 @@ namespace UnrealBuildTool
 			// thread start time
 			Action.StartTime = DateTimeOffset.Now;
 
+			string Args = Action.CommandArguments;
+#if NET_CORE
+			// Process Arguments follow windows conventions in .NET Core
+			// Which means single quotes ' are not considered quotes.
+			// see [REDACTED_PROFILE_URL]/runtime/issues/29857
+			Args = Args.Replace('\'', '\"');
+#endif
 			// Create the action's process.
 			ProcessStartInfo ActionStartInfo = new ProcessStartInfo();
 			ActionStartInfo.WorkingDirectory = Action.WorkingDirectory.FullName;
 			ActionStartInfo.FileName = Action.CommandPath.FullName;
-			ActionStartInfo.Arguments = Action.CommandArguments;
+			ActionStartInfo.Arguments = Args;
 			ActionStartInfo.UseShellExecute = false;
 			ActionStartInfo.RedirectStandardInput = false;
 			ActionStartInfo.RedirectStandardOutput = false;
diff --git a/Engine/Source/Programs/UnrealBuildTool/System/Utils.cs b/Engine/Source/Programs/UnrealBuildTool/System/Utils.cs
--- a/Engine/Source/Programs/UnrealBuildTool/System/Utils.cs
+++ b/Engine/Source/Programs/UnrealBuildTool/System/Utils.cs
@@ -437,7 +437,12 @@ namespace UnrealBuildTool
 		/// <param name="LogOutput">Whether to also log standard output and standard error</param>
 		public static string RunLocalProcessAndReturnStdOut(string Command, string Args, out int ExitCode, bool LogOutput = false)
 		{
-			//LUMIN_MERGE
+#if NET_CORE
+			// Process Arguments follow windows conventions in .NET Core
+			// Which means single quotes ' are not considered quotes.
+			// see [REDACTED_PROFILE_URL]/runtime/issues/29857
+			Args = Args.Replace('\'', '\"');
+#endif
 			ProcessStartInfo StartInfo = new ProcessStartInfo(Command, Args);
 			StartInfo.UseShellExecute = false;
 			StartInfo.RedirectStandardOutput = true;
diff --git a/Engine/Source/Programs/UnrealBuildTool/UnrealBuildTool.cs b/Engine/Source/Programs/UnrealBuildTool/UnrealBuildTool.cs
--- a/Engine/Source/Programs/UnrealBuildTool/UnrealBuildTool.cs
+++ b/Engine/Source/Programs/UnrealBuildTool/UnrealBuildTool.cs
@@ -49,7 +49,14 @@ namespace UnrealBuildTool
 		/// The path to UBT
 		/// </summary>
 #if NET_CORE
+		// In net core the actual running entrypoint is the dll, but its easier to run the apphost executable instead to avoid having to run via the dotnet command
+#if WINDOWS
+		// For windows a executable has a .exe extension
 		public static readonly FileReference UnrealBuildToolPath = FileReference.FindCorrectCase(new FileReference(Path.ChangeExtension(Assembly.GetExecutingAssembly().GetOriginalLocation(), "exe")));
+#else
+		// For linux and mac executables have no extension
+		public static readonly FileReference UnrealBuildToolPath = FileReference.FindCorrectCase(new FileReference(Path.ChangeExtension(Assembly.GetExecutingAssembly().GetOriginalLocation(), null)));
+#endif
 #else
 		public static readonly FileReference UnrealBuildToolPath = FileReference.FindCorrectCase(new FileReference(Assembly.GetExecutingAssembly().GetOriginalLocation()));
 #endif
diff --git a/Engine/Source/Programs/UnrealBuildTool/UnrealBuildToolCore.csproj b/Engine/Source/Programs/UnrealBuildTool/UnrealBuildToolCore.csproj
--- a/Engine/Source/Programs/UnrealBuildTool/UnrealBuildToolCore.csproj
+++ b/Engine/Source/Programs/UnrealBuildTool/UnrealBuildToolCore.csproj
@@ -17,7 +17,27 @@
     <ApplicationManifest>app.manifest</ApplicationManifest>
     <AssemblyName>UnrealBuildTool</AssemblyName>
     <RootNamespace>UnrealBuildTool</RootNamespace>
+    <AppendRuntimeIdentifierToOutputPath>false</AppendRuntimeIdentifierToOutputPath>
   </PropertyGroup>
+  <Choose>
+    <When Condition="$([MSBuild]::IsOSPlatform('OSX'))">
+      <PropertyGroup>
+        <!-- Set runtie identifier for osx so that it produces a app host executable (consistency with the other platforms) -->
+        <RuntimeIdentifier>osx-x64</RuntimeIdentifier>
+      </PropertyGroup>
+    </When>
+    <!-- The default runtime identifiers are sufficent for Linux and Windows
+    <When Condition="$([MSBuild]::IsOSPlatform('Linux'))">
+      <PropertyGroup>
+        <RuntimeIdentifier>linux-x64</RuntimeIdentifier>
+      </PropertyGroup>
+    </When>
+    <When Condition="$([MSBuild]::IsOSPlatform('Windows'))">
+      <PropertyGroup>
+        <RuntimeIdentifier>win-x64</RuntimeIdentifier>
+      </PropertyGroup>
+    </When>-->
+  </Choose>
   <PropertyGroup Condition=" '$(Configuration)|$(Platform)' == 'Development|AnyCPU' ">
     <Optimize>true</Optimize>
     <OutputPath>..\..\..\Binaries\DotNET\UnrealBuildTool</OutputPath>
