-Compile for UAT is now handled by the batchfiles rather then inside UAT itself when running netcore version (as netcore does not support shadow copy so we can not update the assemblies we are running)

Also removed usages of automationtool launcher in netcore as we no longer need it.

diff --git a/Engine/Build/BatchFiles/RunUAT.bat b/Engine/Build/BatchFiles/RunUAT.bat
--- a/Engine/Build/BatchFiles/RunUAT.bat
+++ b/Engine/Build/BatchFiles/RunUAT.bat
@@ -11,11 +11,16 @@ echo Running AutomationTool...
 
 set UATExecutable=AutomationToolLauncher.exe
 set UATDirectory=Binaries\DotNET\
+set UATCompileArg=-compile
+
 if "%UE_USE_DOTNET%" == "1" (
-	set UATDirectory=Binaries\DotNET\AutomationToolLauncher
+	REM We do not use the automaton tool launcher for dotnet builds
+	set UATExecutable=AutomationTool.exe
+	set UATDirectory=Binaries\DotNET\AutomationTool
+	REM -compile is not supported with netcore
+	set UATCompileArg=
 )
 
-set UATCompileArg=-compile
 
 rem ## Change the CWD to /Engine. 
 pushd "%~dp0..\..\"
@@ -53,9 +58,8 @@ rem ## Verify that dotnet is present
 call "%~dp0GetDotnetPath.bat"
 if errorlevel 1 goto Error_NoDotnetSDK
 
-dotnet build Source\Programs\AutomationToolLauncher\AutomationToolLauncherCore.csproj -c Development 
-if errorlevel 1 goto Error_UATCompileFailed
-dotnet build Source\Programs\AutomationTool\AutomationToolCore.csproj -c Development 
+echo Building AutomationTool...
+dotnet build Source\Programs\AutomationTool\AutomationToolCore.csproj -c Development -v quiet 
 if errorlevel 1 goto Error_UATCompileFailed
 goto DoRunUAT
 
@@ -65,10 +69,6 @@ rem ## ok, well it doesn't look like visual studio is installed, let's try runni
 set UATCompileArg=
 if "%UE_USE_DOTNET%" == "1" (
 	if not exist Binaries\DotNET\AutomationTool\AutomationTool.exe goto Error_NoFallbackExecutable
-	if not exist Binaries\DotNET\AutomationToolLauncher\AutomationToolLauncher.exe (
-		set UATExecutable=AutomationTool.exe
-		set UATDirectory=Binaries\DotNET\AutomationTool
-	) 
 ) else (
 	if not exist Binaries\DotNET\AutomationTool.exe goto Error_NoFallbackExecutable
 	if not exist Binaries\DotNET\AutomationToolLauncher.exe set UATExecutable=AutomationTool.exe
diff --git a/Engine/Source/Programs/AutomationTool/AutomationUtils/Automation.cs b/Engine/Source/Programs/AutomationTool/AutomationUtils/Automation.cs
--- a/Engine/Source/Programs/AutomationTool/AutomationUtils/Automation.cs
+++ b/Engine/Source/Programs/AutomationTool/AutomationUtils/Automation.cs
@@ -97,7 +97,9 @@ namespace AutomationTool
 		public static CommandLineArg NoSubmit = new CommandLineArg("-NoSubmit");
 		public static CommandLineArg NoP4 = new CommandLineArg("-NoP4");
 		public static CommandLineArg P4 = new CommandLineArg("-P4");
+#if !NET_CORE
         public static CommandLineArg Compile = new CommandLineArg("-Compile");
+#endif
 		public static CommandLineArg IgnoreDependencies = new CommandLineArg("-IgnoreDependencies");
 		/// <summary>
 		/// This command is LEGACY because we used to run UAT.exe to compile scripts by default.
diff --git a/Engine/Source/Programs/AutomationTool/AutomationUtils/ProjectUtils.cs b/Engine/Source/Programs/AutomationTool/AutomationUtils/ProjectUtils.cs
--- a/Engine/Source/Programs/AutomationTool/AutomationUtils/ProjectUtils.cs
+++ b/Engine/Source/Programs/AutomationTool/AutomationUtils/ProjectUtils.cs
@@ -713,6 +713,7 @@ namespace AutomationTool
 
 				// Check if the scripts require compilation
 				bool DoNotCompile = false;
+
 				if (!CommandUtils.IsBuildMachine && !CheckIfScriptAssemblyIsOutOfDate(TargetsDllFilename, TargetScripts))
 				{
 					Log.TraceVerbose("Targets DLL {0} is up to date.", TargetsDllFilename);
@@ -744,9 +745,12 @@ namespace AutomationTool
 
 			List<string> ReferencedAssemblies = new List<string>() 
 					{ 
+#if !NET_CORE 
+// for net core we already define the system references for DynamicCompilation
 						"System.dll", 
 						"System.Core.dll", 
 						"System.Xml.dll", 
+#endif
 						typeof(UnrealBuildTool.PlatformExports).Assembly.Location
 					};
 			List<string> PreprocessorDefinitions = RulesAssembly.GetPreprocessorDefinitions();
diff --git a/Engine/Source/Programs/AutomationTool/AutomationUtils/ScriptCompiler.cs b/Engine/Source/Programs/AutomationTool/AutomationUtils/ScriptCompiler.cs
--- a/Engine/Source/Programs/AutomationTool/AutomationUtils/ScriptCompiler.cs
+++ b/Engine/Source/Programs/AutomationTool/AutomationUtils/ScriptCompiler.cs
@@ -68,6 +68,8 @@ namespace AutomationTool
 			Parallel.For(0, ProjectFiles.Count, Idx => Projects[Idx] = CsProjectInfo.Read(ProjectFiles[Idx], MsBuildProperties));
 			Log.TraceLog("Parsed project files in {0:0.000}s", ParsingTimer.Elapsed.TotalSeconds);
 
+			// net core does not support shadow copying, as such we expect the compile to happen before running UAT via dotnet build
+#if !NET_CORE
 			// Compile only if not disallowed.
 			if (GlobalCommandLine.Compile && !String.IsNullOrEmpty(CommandUtils.CmdEnv.MsBuildExe))
 			{
@@ -78,6 +80,7 @@ namespace AutomationTool
 				}
 				CompileAutomationProjects(CompileProjects, MsBuildProperties);
 			}
+#endif
 
 			// Get all the build artifacts
 			BuildProducts = new HashSet<FileReference>();
diff --git a/Engine/Source/Programs/AutomationTool/AutomationUtils/UE4Build.cs b/Engine/Source/Programs/AutomationTool/AutomationUtils/UE4Build.cs
--- a/Engine/Source/Programs/AutomationTool/AutomationUtils/UE4Build.cs
+++ b/Engine/Source/Programs/AutomationTool/AutomationUtils/UE4Build.cs
@@ -1261,10 +1261,12 @@ namespace AutomationTool
 		/// </summary>
 		public void AddUBTFilesToBuildProducts()
 		{
+#if !NET_CORE
 			if (!GlobalCommandLine.Compile)
 			{
 				CommandUtils.LogVerbose("We are being asked to copy the UBT build products, but we are running precompiled, so this does not make much sense.");
 			}
+#endif
 
 			var UBTLocation = Path.GetDirectoryName(GetUBTExecutable());
 			var UBTFiles = new List<string>(new string[] 
