Use a submitted dotnet sdk on windows (like we do for the other platforms) unless UE_USE_SYSTEM_DOTNET is set.

diff --git a/Engine/Build/BatchFiles/GetDotnetPath.bat b/Engine/Build/BatchFiles/GetDotnetPath.bat
--- a/Engine/Build/BatchFiles/GetDotnetPath.bat
+++ b/Engine/Build/BatchFiles/GetDotnetPath.bat
@@ -4,7 +4,24 @@ setlocal
 rem ## Unreal Engine 5 utility script
 rem ## Copyright Epic Games, Inc. All Rights Reserved.
 rem ##
-rem ## This script verifies that dotnet sdk is installed and a new enough SDK is present.
+rem ## This script verifies that dotnet sdk is installed and a new enough SDK is present, alternatively setups up the submitted SDK for use.
+
+rem if UE_USE_SYSTEM_DOTNET we assume a installed dotnet is present
+if "%UE_USE_SYSTEM_DOTNET%" == "1" goto verify_dotnet
+
+rem add the dotnet sdk in the repo as the current dotnet sdk
+set UE_DOTNET_DIR=%~dp0..\..\Binaries\ThirdParty\DotNet\Windows
+set PATH=%UE_DOTNET_DIR%;%PATH%
+set DOTNET_ROOT=%UE_DOTNET_DIR%
+
+rem Disables use of any installed version of dotnet
+set DOTNET_MULTILEVEL_LOOKUP=0
+
+rem skip dotnet verification when using our submitted sdk as we know it is up to date
+ECHO Using Dotnet SDK from sync
+goto Succeeded
+
+:verify_dotnet
 
 for /f "delims=" %%i in ('where dotnet') do (
 	REM Dotnet exists
@@ -34,6 +51,7 @@ for /f "tokens=1,* delims= " %%I in ('dotnet --list-sdks') do (
 
 				set FOUND_MAJOR=%%X
 				set FOUND_MINOR=%%Y
+				ECHO Found Dotnet SDK version: %FOUND_MAJOR%.%FOUND_MINOR%
 				goto Succeeded
 			)
 		)
@@ -42,6 +60,7 @@ for /f "tokens=1,* delims= " %%I in ('dotnet --list-sdks') do (
 			REM If the major version is greater then what we require then this sdk is good enough
 			set FOUND_MAJOR=%%X
 			set FOUND_MINOR=%%Y
+			ECHO Found Dotnet SDK version: %FOUND_MAJOR%.%FOUND_MINOR%
 			goto Succeeded
 		)
 
@@ -53,6 +72,5 @@ REM Dotnet is installed but the sdk present is to old
 exit /B 1
 
 :Succeeded
-ECHO Found Dotnet SDK version: %FOUND_MAJOR%.%FOUND_MINOR%
 exit /B 0
 	
\ No newline at end of file
