Fixed issue with submitted dotnet sdk not being used on windows when building UAT

diff --git a/Engine/Build/BatchFiles/GetDotnetPath.bat b/Engine/Build/BatchFiles/GetDotnetPath.bat
--- a/Engine/Build/BatchFiles/GetDotnetPath.bat
+++ b/Engine/Build/BatchFiles/GetDotnetPath.bat
@@ -1,5 +1,4 @@
 @echo off
-setlocal
 
 rem ## Unreal Engine 5 utility script
 rem ## Copyright Epic Games, Inc. All Rights Reserved.
