Removed passing UAT compile argument for net core, we handle this in the shell scripts.

diff --git a/Engine/Build/BatchFiles/RunUAT.sh b/Engine/Build/BatchFiles/RunUAT.sh
--- a/Engine/Build/BatchFiles/RunUAT.sh
+++ b/Engine/Build/BatchFiles/RunUAT.sh
@@ -57,8 +57,8 @@ fi
 
 EnvironmentType=-mono
 if [ ${UE_USE_DOTNET:=0} -ne 0 ]; then
-	EnvironmentType=-dotnet
-	UATDirectory=Binaries/DotNET/AutomationTool
+EnvironmentType=-dotnet
+UATDirectory=Binaries/DotNET/AutomationTool
 fi
 
 if [ "$(uname)" = "Darwin" ]; then
@@ -79,18 +79,17 @@ if [ "$UATCompileArg" = "-compile" ]; then
 		UATCompileArg=
 	else
 		if [ ${UE_USE_DOTNET:=0} -ne 0 ]; then
-			echo Building AutomationTool...
+		echo Building AutomationTool...
 			dotnet msbuild -restore Source/Programs/AutomationTool/AutomationToolCore.csproj /property:Configuration=Development /property:AutomationToolProjectOnly=true /verbosity:quiet
-			if [ $? -ne 0 ]; then
-				echo RunUAT ERROR: AutomationTool failed to compile.
-				exit 1
-			fi
-
-			echo Building AutomationTool Plugins...
+		if [ $? -ne 0 ]; then
+			echo RunUAT ERROR: AutomationTool failed to compile.
+			exit 1
+		fi
+		echo Building AutomationTool Plugins...
 			dotnet msbuild -restore Source/Programs/AutomationTool/AutomationTool.proj /property:Configuration=Development /verbosity:quiet
-			if [ $? -ne 0 ]; then
-				echo RunUAT ERROR: AutomationTool plugins failed to compile.
-				exit 1
+		if [ $? -ne 0 ]; then
+			echo RunUAT ERROR: AutomationTool plugins failed to compile.
+			exit 1
 			fi
 		else
 			# mono 5.0 and up include msbuild
@@ -99,10 +98,8 @@ if [ "$UATCompileArg" = "-compile" ]; then
 			else
 				BUILD_TOOL=xbuild
 			fi
-
 			ARGS="/p:Configuration=Development /p:Platform=AnyCPU /verbosity:quiet /nologo /p:NoWarn=1591 /property:AutomationToolProjectOnly=true"
 			ARGS="${ARGS} /p:TargetFrameworkProfile="
-
 			echo "$BUILD_TOOL Source/Programs/AutomationTool/AutomationTool.csproj $ARGS"
 			$BUILD_TOOL Source/Programs/AutomationTool/AutomationTool.csproj $ARGS
 			# make sure it succeeded
@@ -125,8 +122,8 @@ else
 fi
 
 if [ ${UE_USE_DOTNET:=0} -ne 0 ]; then
-	echo Start UAT: AutomationTool "${Args[@]}"
-	env uebp_LogFolder="$LogDir" ./AutomationTool "${Args[@]}" $UATCompileArg
+echo Start UAT: AutomationTool "${Args[@]}"
+	env uebp_LogFolder="$LogDir" ./AutomationTool "${Args[@]}"
 else
 	# if we are running under UE, we need to run this with the term handler (otherwise canceling a UAT job from the editor
 	# can leave mono, etc running in the background, which means we need the PID so we 
