Fixed dotnet sdk version output when using the system dotnet sdk

diff --git a/Engine/Build/BatchFiles/GetDotnetPath.bat b/Engine/Build/BatchFiles/GetDotnetPath.bat
--- a/Engine/Build/BatchFiles/GetDotnetPath.bat
+++ b/Engine/Build/BatchFiles/GetDotnetPath.bat
@@ -50,7 +50,7 @@ for /f "tokens=1,* delims= " %%I in ('dotnet --list-sdks') do (
 
 				set FOUND_MAJOR=%%X
 				set FOUND_MINOR=%%Y
-				ECHO Found Dotnet SDK version: %FOUND_MAJOR%.%FOUND_MINOR%
+				ECHO Found Dotnet SDK version: %%X.%%Y
 				goto Succeeded
 			)
 		)
@@ -59,7 +59,7 @@ for /f "tokens=1,* delims= " %%I in ('dotnet --list-sdks') do (
 			REM If the major version is greater then what we require then this sdk is good enough
 			set FOUND_MAJOR=%%X
 			set FOUND_MINOR=%%Y
-			ECHO Found Dotnet SDK version: %FOUND_MAJOR%.%FOUND_MINOR%
+			ECHO Found Dotnet SDK version: %%X.%%Y
 			goto Succeeded
 		)
 
