Fixed dotnet environment on Mac, the bundled dotnet sdk setup was incorrect.

diff --git a/Engine/Build/BatchFiles/Mac/SetupDotnet.sh b/Engine/Build/BatchFiles/Mac/SetupDotnet.sh
--- a/Engine/Build/BatchFiles/Mac/SetupDotnet.sh
+++ b/Engine/Build/BatchFiles/Mac/SetupDotnet.sh
@@ -1,6 +1,9 @@
 #!/bin/bash
 # Copyright Epic Games, Inc. All Rights Reserved.
 
+START_DIR=`pwd`
+cd "$1"
+
 IS_DOTNET_INSTALLED=0
 DOTNET_VERSION_PATH=$(command -v dotnet) || true
 
@@ -8,25 +11,27 @@ if [ "$UE_USE_SYSTEM_DOTNET" == "1" ] && [ ! $DOTNET_VERSION_PATH == "" ] && [ -
 	# If dotnet is installed, check that it has a new enough version of the SDK
 	DOTNET_SDKS=(`dotnet --list-sdks | grep -P "(\d*)\.(\d*)\..* \[(.*)\]"`)
 	for DOTNET_SDK in $DOTNET_SDKS
-    do
-        if [ ${DOTNET_SDK[0]} -ge 3 ]; then
-            if [ ${DOTNET_SDK[1]} -ge 1 ]; then
-                IS_DOTNET_INSTALLED=1
-            fi
-        fi
-    done
-    if [ $IS_DOTNET_INSTALLED -eq 0 ]; then
-    echo Unable to find installed dotnet sdk of version 3.1 or newer
-    fi
+	do
+		if [ ${DOTNET_SDK[0]} -ge 3 ]; then
+			if [ ${DOTNET_SDK[1]} -ge 1 ]; then
+				IS_DOTNET_INSTALLED=1
+			fi
+		fi
+	done
+	if [ $IS_DOTNET_INSTALLED -eq 0 ]; then
+	echo Unable to find installed dotnet sdk of version 3.1 or newer
+	fi
 fi
 
 # Setup bundled Dotnet if cannot use installed one
 if [ $IS_DOTNET_INSTALLED -eq 0 ]; then
 	echo Setting up bundled DotNet SDK
-    CUR_DIR=`pwd`
+	CUR_DIR=`pwd`
 	export UE_DOTNET_DIR=$CUR_DIR/../../../Binaries/ThirdParty/DotNet/Mac
 	export PATH=$UE_DOTNET_DIR:$PATH
-    export DOTNET_ROOT=$UE_DOTNET_DIR
+	export DOTNET_ROOT=$UE_DOTNET_DIR
 else
 	export IS_DOTNET_INSTALLED=$IS_DOTNET_INSTALLED
-fi
\ No newline at end of file
+fi
+
+cd "$START_DIR"
\ No newline at end of file
diff --git a/Engine/Build/BatchFiles/Mac/SetupEnvironment.sh b/Engine/Build/BatchFiles/Mac/SetupEnvironment.sh
--- a/Engine/Build/BatchFiles/Mac/SetupEnvironment.sh
+++ b/Engine/Build/BatchFiles/Mac/SetupEnvironment.sh
@@ -25,7 +25,7 @@ do
             ;;
         -dotnet) 
             # echo "setting up dotnet"
-            source $2/SetupDotnet.sh "$2"
+            source "$2/SetupDotnet.sh" "$2"
             shift
             ;;
         -*) echo "bad argument $1"
