When using UBT in dotnet, unset TargetName when starting build from xcode as this also has meaning to dotnet breaking the build of UBT.

diff --git a/Engine/Build/BatchFiles/Mac/XcodeBuild.sh b/Engine/Build/BatchFiles/Mac/XcodeBuild.sh
--- a/Engine/Build/BatchFiles/Mac/XcodeBuild.sh
+++ b/Engine/Build/BatchFiles/Mac/XcodeBuild.sh
@@ -6,6 +6,8 @@
 
 # Setup Environment
 if [ ${UE_USE_DOTNET:=0} -ne 0 ]; then
+  # remove environment variable passed from xcode which also has meaning to dotnet, breaking the build
+  unset TARGETNAME
   source  Engine/Build/BatchFiles/Mac/SetupEnvironment.sh -dotnet Engine/Build/BatchFiles/Mac
 else
   source Engine/Build/BatchFiles/Mac/SetupEnvironment.sh -mono Engine/Build/BatchFiles/Mac
