Fixed dotnet_root not being set when executing UBT if you do not have a dotnet runtime installed.

diff --git a/Engine/Build/BatchFiles/Build.bat b/Engine/Build/BatchFiles/Build.bat
--- a/Engine/Build/BatchFiles/Build.bat
+++ b/Engine/Build/BatchFiles/Build.bat
@@ -5,6 +5,10 @@ REM The %~dp0 specifier resolves to the path to the directory where this .bat is
 REM We use this so that regardless of where the .bat file was executed from, we can change to
 REM directory relative to where we know the .bat is stored.
 pushd "%~dp0\..\..\Source"
+if "%UE_USE_DOTNET%" == "1" (
+	rem ## Verify that dotnet is present
+	call "%~dp0GetDotnetPath.bat"
+)
 
 REM %1 is the game name
 REM %2 is the platform name
