Improved output of using bundled dotnet sdk

diff --git a/Engine/Build/BatchFiles/GetDotnetPath.bat b/Engine/Build/BatchFiles/GetDotnetPath.bat
--- a/Engine/Build/BatchFiles/GetDotnetPath.bat
+++ b/Engine/Build/BatchFiles/GetDotnetPath.bat
@@ -17,7 +17,7 @@ rem Disables use of any installed version of dotnet
 set DOTNET_MULTILEVEL_LOOKUP=0
 
 rem skip dotnet verification when using our submitted sdk as we know it is up to date
-ECHO Using Dotnet SDK from sync
+ECHO Using bundled DotNet SDK
 goto Succeeded
 
 :verify_dotnet
