Temporarily increased msbuild verbosity to debug issues when starting incremental builds in horde

diff --git a/Engine/Build/BatchFiles/RunUAT.bat b/Engine/Build/BatchFiles/RunUAT.bat
--- a/Engine/Build/BatchFiles/RunUAT.bat
+++ b/Engine/Build/BatchFiles/RunUAT.bat
@@ -18,7 +18,7 @@ rem ## Change the CWD to /Engine.
 pushd "%~dp0..\..\"
 if not exist Build\BatchFiles\RunUAT.bat goto Error_BatchFileInWrongLocation
 
-set MSBUILD_LOGLEVEL=quiet
+set MSBUILD_LOGLEVEL=minimal
 for %%P in (%*) do if /I "%%P" == "-msbuild-verbose" set MSBUILD_LOGLEVEL=normal
 
 rem ## Use the pre-compiled UAT scripts if -nocompile is specified in the command line
