Build UBT before running UAT, this should not be needed but testing to see if this resolves our incremental horde build issue.

diff --git a/Engine/Build/BatchFiles/RunUAT.bat b/Engine/Build/BatchFiles/RunUAT.bat
--- a/Engine/Build/BatchFiles/RunUAT.bat
+++ b/Engine/Build/BatchFiles/RunUAT.bat
@@ -39,6 +39,9 @@ rem ## Verify that dotnet is present
 call "%~dp0GetDotnetPath.bat"
 if errorlevel 1 goto Error_NoDotnetSDK
 
+echo Building UnrealBuildTool...
+dotnet build Source\Programs\UnrealBuildTool\UnrealBuildTool.csproj -c Development -v %MSBUILD_LOGLEVEL% --nologo 1>nul
+if errorlevel 1 goto Error_UATCompileFailed
 echo Building AutomationTool...
 dotnet msbuild /restore /property:Configuration=Development /property:AutomationToolProjectOnly=true /verbosity:%MSBUILD_LOGLEVEL% Source\Programs\AutomationTool\AutomationTool.csproj
 if errorlevel 1 goto Error_UATCompileFailed
