Reduced verbosity on RunUAT back to expected levels and removed msbuild logo when invoking it to clean up the output a bit.

diff --git a/Engine/Build/BatchFiles/RunUAT.bat b/Engine/Build/BatchFiles/RunUAT.bat
--- a/Engine/Build/BatchFiles/RunUAT.bat
+++ b/Engine/Build/BatchFiles/RunUAT.bat
@@ -18,7 +18,7 @@ rem ## Change the CWD to /Engine.
 pushd "%~dp0..\..\"
 if not exist Build\BatchFiles\RunUAT.bat goto Error_BatchFileInWrongLocation
 
-set MSBUILD_LOGLEVEL=minimal
+set MSBUILD_LOGLEVEL=quiet
 for %%P in (%*) do if /I "%%P" == "-msbuild-verbose" set MSBUILD_LOGLEVEL=normal
 
 rem ## Use the pre-compiled UAT scripts if -nocompile is specified in the command line
@@ -43,10 +43,10 @@ echo Building UnrealBuildTool...
 dotnet build Source\Programs\UnrealBuildTool\UnrealBuildTool.csproj -c Development -v %MSBUILD_LOGLEVEL% --nologo 1>nul
 if errorlevel 1 goto Error_UATCompileFailed
 echo Building AutomationTool...
-dotnet msbuild /restore /property:Configuration=Development /property:AutomationToolProjectOnly=true /verbosity:%MSBUILD_LOGLEVEL% Source\Programs\AutomationTool\AutomationTool.csproj
+dotnet msbuild /restore /property:Configuration=Development /nologo /property:AutomationToolProjectOnly=true /verbosity:%MSBUILD_LOGLEVEL% Source\Programs\AutomationTool\AutomationTool.csproj
 if errorlevel 1 goto Error_UATCompileFailed
 echo Building AutomationTool Plugins...
-dotnet msbuild /restore /property:Configuration=Development /verbosity:%MSBUILD_LOGLEVEL% Source\Programs\AutomationTool\AutomationTool.proj
+dotnet msbuild /restore /property:Configuration=Development /nologo /verbosity:%MSBUILD_LOGLEVEL% Source\Programs\AutomationTool\AutomationTool.proj
 if errorlevel 1 goto Error_UATCompileFailed
 goto DoRunUAT
 
