Switched over to dotnet msbuild from dotnet build (which just calls msbuild anyway).
This has less verbose output when a build is successful (nothing) but still outputs errors when something is wrong, which means we can remove the redirection of stdout so we actually can see compile errors from UBT when running GenerateProjects.

diff --git a/Engine/Build/BatchFiles/GenerateProjectFiles.bat b/Engine/Build/BatchFiles/GenerateProjectFiles.bat
--- a/Engine/Build/BatchFiles/GenerateProjectFiles.bat
+++ b/Engine/Build/BatchFiles/GenerateProjectFiles.bat
@@ -48,11 +48,11 @@ fc /b ..\Intermediate\Build\UnrealBuildToolFiles.txt ..\Intermediate\Build\Unrea
 if not errorlevel 1 goto SkipClean
 copy /y ..\Intermediate\Build\UnrealBuildToolFiles.txt ..\Intermediate\Build\UnrealBuildToolPrevFiles.txt >nul
 
-dotnet clean Programs\UnrealBuildTool\UnrealBuildTool.csproj -c Development -v quiet --nologo 1>nul
+dotnet msbuild /target:clean /property:Configuration=Development /nologo Programs\UnrealBuildTool\UnrealBuildTool.csproj /verbosity:quiet
 
 :SkipClean
 echo Building UnrealBuildTool...
-dotnet build Programs\UnrealBuildTool\UnrealBuildTool.csproj -c Development -v quiet --nologo 1>nul
+dotnet msbuild /restore /target:build /property:Configuration=Development /nologo Programs\UnrealBuildTool\UnrealBuildTool.csproj /verbosity:quiet
 
 if errorlevel 1 goto Error_UBTCompileFailed
 
diff --git a/Engine/Build/BatchFiles/Linux/GenerateProjectFiles.sh b/Engine/Build/BatchFiles/Linux/GenerateProjectFiles.sh
--- a/Engine/Build/BatchFiles/Linux/GenerateProjectFiles.sh
+++ b/Engine/Build/BatchFiles/Linux/GenerateProjectFiles.sh
@@ -23,7 +23,7 @@ fi
 source "$BASE_PATH/SetupEnvironment.sh" -dotnet "$BASE_PATH"
 
 if [ -f "$BASE_PATH/../../../Source/Programs/UnrealBuildTool/UnrealBuildTool.csproj" ]; then
-  dotnet build $BASE_PATH/../../../Source/Programs/UnrealBuildTool/UnrealBuildTool.csproj -c Development -v quiet
+  dotnet msbuild /restore /target:build /property:Configuration=Development /nologo $BASE_PATH/../../Source/Programs/UnrealBuildTool/UnrealBuildTool.csproj /verbosity:quiet
   if [ $? -ne 0 ]; then
     echo GenerateProjectFiles ERROR: Failed to build UnrealBuildTool
     exit 1
diff --git a/Engine/Build/BatchFiles/Mac/GenerateProjectFiles.sh b/Engine/Build/BatchFiles/Mac/GenerateProjectFiles.sh
--- a/Engine/Build/BatchFiles/Mac/GenerateProjectFiles.sh
+++ b/Engine/Build/BatchFiles/Mac/GenerateProjectFiles.sh
@@ -23,7 +23,7 @@ fi
 source "$BASE_PATH/SetupEnvironment.sh" -dotnet "$BASE_PATH"
 
 if [ -f "$BASE_PATH/../../../Source/Programs/UnrealBuildTool/UnrealBuildTool.csproj" ]; then
-  dotnet build $BASE_PATH/../../../Source/Programs/UnrealBuildTool/UnrealBuildTool.csproj -c Development -v quiet
+  dotnet msbuild /restore /target:build /property:Configuration=Development /nologo $BASE_PATH/../../Source/Programs/UnrealBuildTool/UnrealBuildTool.csproj /verbosity:quiet
 
   if [ $? -ne 0 ]; then
     echo GenerateProjectFiles ERROR: Failed to build UnrealBuildTool
diff --git a/Engine/Build/BatchFiles/RunUAT.bat b/Engine/Build/BatchFiles/RunUAT.bat
--- a/Engine/Build/BatchFiles/RunUAT.bat
+++ b/Engine/Build/BatchFiles/RunUAT.bat
@@ -40,7 +40,7 @@ call "%~dp0GetDotnetPath.bat"
 if errorlevel 1 goto Error_NoDotnetSDK
 
 echo Building UnrealBuildTool...
-dotnet build Source\Programs\UnrealBuildTool\UnrealBuildTool.csproj -c Development -v %MSBUILD_LOGLEVEL% --nologo 1>nul
+dotnet msbuild /restore /target:build /property:Configuration=Development /nologo Source\Programs\UnrealBuildTool\UnrealBuildTool.csproj /verbosity:%MSBUILD_LOGLEVEL%
 if errorlevel 1 goto Error_UATCompileFailed
 echo Building AutomationTool...
 dotnet msbuild /restore /property:Configuration=Development /nologo /property:AutomationToolProjectOnly=true /verbosity:%MSBUILD_LOGLEVEL% Source\Programs\AutomationTool\AutomationTool.csproj
diff --git a/Engine/Build/BatchFiles/RunUBT.bat b/Engine/Build/BatchFiles/RunUBT.bat
--- a/Engine/Build/BatchFiles/RunUBT.bat
+++ b/Engine/Build/BatchFiles/RunUBT.bat
@@ -40,10 +40,11 @@ fc /b ..\Intermediate\Build\UnrealBuildToolFiles.txt ..\Intermediate\Build\Unrea
 if not errorlevel 1 goto SkipClean
 copy /y ..\Intermediate\Build\UnrealBuildToolFiles.txt ..\Intermediate\Build\UnrealBuildToolPrevFiles.txt >nul
 
-dotnet clean Programs\UnrealBuildTool\UnrealBuildTool.csproj -c Development -v quiet --nologo 1>nul
+dotnet msbuild /target:clean /property:Configuration=Development /nologo Programs\UnrealBuildTool\UnrealBuildTool.csproj /verbosity:quiet
+
 :SkipClean
 echo Building UnrealBuildTool...
-dotnet build Programs\UnrealBuildTool\UnrealBuildTool.csproj -c Development -v quiet --nologo 1>nul
+dotnet msbuild /restore /target:build /property:Configuration=Development /nologo Programs\UnrealBuildTool\UnrealBuildTool.csproj /verbosity:quiet
 if errorlevel 1 goto Error_UBTCompileFailed
 
 rem ## Run UnrealBuildTool to generate Visual Studio solution and project files
diff --git a/Engine/Build/BatchFiles/RunUBT.sh b/Engine/Build/BatchFiles/RunUBT.sh
--- a/Engine/Build/BatchFiles/RunUBT.sh
+++ b/Engine/Build/BatchFiles/RunUBT.sh
@@ -27,7 +27,8 @@ if [ "$(uname)" = "Linux" ]; then
 fi
 
 if [ -f "$BASE_PATH/../../Source/Programs/UnrealBuildTool/UnrealBuildTool.csproj" ]; then 
-  dotnet build $BASE_PATH/../../Source/Programs/UnrealBuildTool/UnrealBuildTool.csproj -c Development -v quiet
+  dotnet msbuild /restore /target:build /property:Configuration=Development /nologo $BASE_PATH/../../Source/Programs/UnrealBuildTool/UnrealBuildTool.csproj /verbosity:quiet
+
   if [ $? -ne 0 ]; then
     echo RunUBT ERROR: Failed to build UnrealBuildTool
     exit 1
