Fixed incorrect path to UBT project

diff --git a/Engine/Build/BatchFiles/Linux/GenerateProjectFiles.sh b/Engine/Build/BatchFiles/Linux/GenerateProjectFiles.sh
--- a/Engine/Build/BatchFiles/Linux/GenerateProjectFiles.sh
+++ b/Engine/Build/BatchFiles/Linux/GenerateProjectFiles.sh
@@ -23,7 +23,7 @@ fi
 source "$BASE_PATH/SetupEnvironment.sh" -dotnet "$BASE_PATH"
 
 if [ -f "$BASE_PATH/../../../Source/Programs/UnrealBuildTool/UnrealBuildTool.csproj" ]; then
-  dotnet msbuild /restore /target:build /property:Configuration=Development /nologo $BASE_PATH/../../Source/Programs/UnrealBuildTool/UnrealBuildTool.csproj /verbosity:quiet
+  dotnet msbuild /restore /target:build /property:Configuration=Development /nologo $BASE_PATH/../../../Source/Programs/UnrealBuildTool/UnrealBuildTool.csproj /verbosity:quiet
   if [ $? -ne 0 ]; then
     echo GenerateProjectFiles ERROR: Failed to build UnrealBuildTool
     exit 1
diff --git a/Engine/Build/BatchFiles/Mac/GenerateProjectFiles.sh b/Engine/Build/BatchFiles/Mac/GenerateProjectFiles.sh
--- a/Engine/Build/BatchFiles/Mac/GenerateProjectFiles.sh
+++ b/Engine/Build/BatchFiles/Mac/GenerateProjectFiles.sh
@@ -1,6 +1,5 @@
 #!/bin/bash
 # Copyright Epic Games, Inc. All Rights Reserved.
-
 echo
 echo Setting up Unreal Engine 5 project files...
 echo
@@ -15,15 +14,15 @@ if [ ! -d "$BASE_PATH/../../../Binaries/DotNET" ]; then
  exit 1
 fi
 
+echo "$BASE_PATH"
 if [ ! -d "$BASE_PATH/../../../Source" ]; then
  echo GenerateProjectFiles ERROR: This script file does not appear to be located inside the Engine/Build/BatchFiles/Mac directory.
  exit 1
 fi
 
 source "$BASE_PATH/SetupEnvironment.sh" -dotnet "$BASE_PATH"
-
 if [ -f "$BASE_PATH/../../../Source/Programs/UnrealBuildTool/UnrealBuildTool.csproj" ]; then
-  dotnet msbuild /restore /target:build /property:Configuration=Development /nologo $BASE_PATH/../../Source/Programs/UnrealBuildTool/UnrealBuildTool.csproj /verbosity:quiet
+  dotnet msbuild /restore /target:build /property:Configuration=Development /nologo $BASE_PATH/../../../Source/Programs/UnrealBuildTool/UnrealBuildTool.csproj /verbosity:quiet
 
   if [ $? -ne 0 ]; then
     echo GenerateProjectFiles ERROR: Failed to build UnrealBuildTool
