Removed node reuse and shared compilation flags for RunUAT.sh, these were not inteded to be submitted. They would disable the worker nodes which means we will now dotnet process hosting msbuild left after the build, which is annoying but a fairly substantial speed up. This is consistent with how RunUAT works on windows.

Also added -msbuild-verbose to RunUAT.sh for consistency with RunUAT.bat

diff --git a/Engine/Build/BatchFiles/RunUAT.sh b/Engine/Build/BatchFiles/RunUAT.sh
--- a/Engine/Build/BatchFiles/RunUAT.sh
+++ b/Engine/Build/BatchFiles/RunUAT.sh
@@ -51,6 +51,13 @@ else
 	UATCompileArg=-compile
 fi
 
+# control toggling of msbuild verbosity for easier debugging
+if echo "${Args[@]}" | grep -q -w -i "\-msbuild-verbose"; then
+	MSBuild_Verbosity=normal
+else
+	MSBuild_Verbosity=quiet
+fi
+
 if [ -f Build/InstalledBuild.txt ]; then
 	UATCompileArg=
 fi
@@ -76,13 +83,13 @@ if [ "$UATCompileArg" = "-compile" ]; then
 		UATCompileArg=
 	else
 		echo Building AutomationTool...
-		dotnet msbuild -restore Source/Programs/AutomationTool/AutomationTool.csproj /property:Configuration=Development /property:AutomationToolProjectOnly=true /verbosity:quiet /nodeReuse:false /p:UseSharedCompilation=false
+		dotnet msbuild -restore Source/Programs/AutomationTool/AutomationTool.csproj /nologo /property:Configuration=Development /property:AutomationToolProjectOnly=true /verbosity:$MSBuild_Verbosity
 		if [ $? -ne 0 ]; then
 			echo RunUAT ERROR: AutomationTool failed to compile.
 			exit 1
 		fi
 		echo Building AutomationTool Plugins...
-		dotnet msbuild -restore Source/Programs/AutomationTool/AutomationTool.proj /property:Configuration=Development /verbosity:quiet  /nodeReuse:false /p:UseSharedCompilation=false
+		dotnet msbuild -restore Source/Programs/AutomationTool/AutomationTool.proj /nologo  /property:Configuration=Development /verbosity:$MSBuild_Verbosity
 		if [ $? -ne 0 ]; then
 			echo RunUAT ERROR: AutomationTool plugins failed to compile.
 			exit 1
