Removed extra call to GetDotnetPath which was caused by a incorrect merge in change 14834347

diff --git a/Engine/Build/BatchFiles/Build.bat b/Engine/Build/BatchFiles/Build.bat
--- a/Engine/Build/BatchFiles/Build.bat
+++ b/Engine/Build/BatchFiles/Build.bat
@@ -5,10 +5,6 @@ REM The %~dp0 specifier resolves to the path to the directory where this .bat is
 REM We use this so that regardless of where the .bat file was executed from, we can change to
 REM directory relative to where we know the .bat is stored.
 pushd "%~dp0\..\..\Source"
-if "%UE_USE_DOTNET%" == "1" (
-	rem ## Verify that dotnet is present
-	call "%~dp0GetDotnetPath.bat"
-)
 
 rem ## Verify that dotnet is present
 call "%~dp0GetDotnetPath.bat"
