Fix merge error breaking iOS remote build

diff --git a/Engine/Build/BatchFiles/Mac/Build.sh b/Engine/Build/BatchFiles/Mac/Build.sh
--- a/Engine/Build/BatchFiles/Mac/Build.sh
+++ b/Engine/Build/BatchFiles/Mac/Build.sh
@@ -7,10 +7,8 @@ cd "`dirname "$0"`/../../../.."
 # Setup Environment for DotNET
 source Engine/Build/BatchFiles/Mac/SetupEnvironment.sh -dotnet Engine/Build/BatchFiles/Mac
 
-
-# Skip UBT and SWC compile step if we're coming in on an SSH connection (ie remote toolchain)
-# or if this is an installed build
-if [ -z "$SSH_CONNECTION" ] && [ ! -f Engine/Build/InstalledBuild.txt ]; then
+# Skip UBT and SWC compile step if this is an installed build.
+if [ ! -f Engine/Build/InstalledBuild.txt ]; then
     # First make sure that the UnrealBuildTool is up-to-date
     if ! dotnet build Engine/Source/Programs/UnrealBuildTool/UnrealBuildTool.csproj -c Development -v quiet; then
 	  echo "Failed to build to build tool (UnrealBuildTool)"
