Hardening validation of command line arguments for RunUAT.bat. Removes need to escape arguments with parentheses.

diff --git a/Engine/Build/BatchFiles/RunUAT.bat b/Engine/Build/BatchFiles/RunUAT.bat
--- a/Engine/Build/BatchFiles/RunUAT.bat
+++ b/Engine/Build/BatchFiles/RunUAT.bat
@@ -6,9 +6,10 @@ rem ##
 rem ## This script is expecting to exist in the UE4/Engine/Build/BatchFiles directory.  It will not work correctly
 rem ## if you copy it to a different location and run it.
 
-setlocal 
+setlocal EnableExtensions
 echo Running AutomationTool...
 
+set SCRIPT_DIR=%~dp0
 set UATExecutable=AutomationTool.exe
 set UATDirectory=Binaries\DotNET\AutomationTool
 REM -compile is not supported with netcore instead we will compile as part of this batch file
@@ -19,11 +20,23 @@ pushd "%~dp0..\..\"
 if not exist Build\BatchFiles\RunUAT.bat goto Error_BatchFileInWrongLocation
 
 set MSBUILD_LOGLEVEL=quiet
-for %%P in (%*) do if /I "%%P" == "-msbuild-verbose" set MSBUILD_LOGLEVEL=normal
+set COMPILE_UAT=1
+set SET_TURNKEY_VARIABLES=1
+
+rem ## Check for any arguments handled by this script, being sensitive to any characters that are treated as delimiters by cmd.exe.
+:ParseArguments
+set ARGUMENT=%1
+if not defined ARGUMENT goto ParseArguments_Done
+set ARGUMENT=%ARGUMENT:"=%
+if /I "%ARGUMENT%" == "-msbuild-verbose" set MSBUILD_LOGLEVEL=normal
+if /I "%ARGUMENT%" == "-nocompile" set COMPILE_UAT=0
+if /I "%ARGUMENT%" == "-noturnkeyvariables" set SET_TURNKEY_VARIABLES=0
+shift
+goto ParseArguments
+:ParseArguments_Done
 
 rem ## Use the pre-compiled UAT scripts if -nocompile is specified in the command line
-for %%P in (%*) do if /I "%%P" == "-nocompile" goto RunPrecompiled
-
+if %COMPILE_UAT%==0 goto RunPrecompiled
 
 rem ## If we're running in an installed build, default to precompiled
 if exist Build\InstalledBuild.txt goto RunPrecompiled
@@ -36,7 +49,7 @@ if not exist Source\Programs\AutomationTool\AutomationTool.csproj goto RunPrecom
 if not exist Source\Programs\AutomationToolLauncher\AutomationToolLauncher.csproj goto RunPrecompiled
 
 rem ## Verify that dotnet is present
-call "%~dp0GetDotnetPath.bat"
+call "%SCRIPT_DIR%GetDotnetPath.bat"
 if errorlevel 1 goto Error_NoDotnetSDK
 
 echo Building UnrealBuildTool...
@@ -63,15 +76,15 @@ pushd %UATDirectory%
 %UATExecutable% %* %UATCompileArg%
 popd
 
-for %%P in (%*) do if /I "%%P" == "-noturnkeyvariables" goto SkipTurnkey
+if %SET_TURNKEY_VARIABLES% == 0 goto SkipTurnkey
 
 rem ## Turnkey needs to update env vars in the calling process so that if it is run multiple times the Sdk env var changes are in effect
-if EXIST ~dp0..\..\Intermediate\Turnkey\PostTurnkeyVariables.bat (
+if EXIST %SCRIPT_DIR%..\..\Intermediate\Turnkey\PostTurnkeyVariables.bat (
 	rem ## We need to endlocal so that the vars in the batch file work. NOTE: Working directory from pushd will be UNDONE here, but since we are about to quit, it's okay
 	endlocal 
 	echo Updating environment variables set by a Turnkey sub-process
-	call ~dp0..\..\Engine\Intermediate\Turnkey\PostTurnkeyVariables.bat
-	del ~dp0..\..\Engine\Intermediate\Turnkey\PostTurnkeyVariables.bat
+	call %SCRIPT_DIR%..\..\Engine\Intermediate\Turnkey\PostTurnkeyVariables.bat
+	del %SCRIPT_DIR%..\..\Engine\Intermediate\Turnkey\PostTurnkeyVariables.bat
 	rem ## setlocal again so that any popd's etc don't have an effect on calling process
 	setlocal
 )
