UnrealBuildTool.exe -> dotnet UnrealBuildTool.dll (1/n)

Build.bat, Clean.bat, GenerateProjectFiles.bat

diff --git a/Engine/Build/BatchFiles/Build.bat b/Engine/Build/BatchFiles/Build.bat
--- a/Engine/Build/BatchFiles/Build.bat
+++ b/Engine/Build/BatchFiles/Build.bat
@@ -13,10 +13,10 @@ REM %1 is the game name
 REM %2 is the platform name
 REM %3 is the configuration name
 
-set UBTPath="..\..\Engine\Binaries\DotNET\UnrealBuildTool\UnrealBuildTool.exe"
+set UBTPath="..\..\Engine\Binaries\DotNET\UnrealBuildTool\UnrealBuildTool.dll"
 
 IF EXIST %UBTPath% (
-		%UBTPath% %*
+		dotnet %UBTPath% %*
 		popd
 
 		REM Ignore exit codes of 2 ("ECompilationResult.UpToDate") from UBT; it's not a failure.
@@ -26,7 +26,7 @@ IF EXIST %UBTPath% (
 		 
 		EXIT /B !ERRORLEVEL!
 ) ELSE (
-	ECHO UnrealBuildTool.exe not found in %UBTPath%
+	ECHO UnrealBuildTool.dll not found in %UBTPath%
 	popd
 	EXIT /B 999
 )
diff --git a/Engine/Build/BatchFiles/Clean.bat b/Engine/Build/BatchFiles/Clean.bat
--- a/Engine/Build/BatchFiles/Clean.bat
+++ b/Engine/Build/BatchFiles/Clean.bat
@@ -32,7 +32,7 @@ if errorlevel 1 goto Error_NoVisualStudioEnvironment
 rem ## Compile UBT if the project file exists
 
 set ProjectFile="Programs\UnrealBuildTool\UnrealBuildTool.csproj"
-set UBTPath="..\..\Engine\Binaries\DotNET\UnrealBuildTool\UnrealBuildTool.exe"
+set UBTPath="..\..\Engine\Binaries\DotNET\UnrealBuildTool\UnrealBuildTool.dll"
 if not exist %ProjectFile% goto NoProjectFile
 dotnet build Programs\UnrealBuildTool\UnrealBuildTool.csproj -c Development -v quiet
 if errorlevel 1 goto Error_UBTCompileFailed
@@ -42,7 +42,7 @@ if errorlevel 1 goto Error_UBTCompileFailed
 rem ## Execute UBT
 :ReadyToClean
 if not exist %UBTPath% goto Error_UBTMissing
-%UBTPath% %* -clean
+dotnet %UBTPath% %* -clean
 goto Exit
 
 
@@ -62,7 +62,7 @@ pause
 goto Exit
 
 :Error_UBTMissing
-echo Clean ERROR: UnrealBuildTool.exe not found in ..\..\Engine\Binaries\DotNET\UnrealBuildTool\UnrealBuildTool.exe 
+echo Clean ERROR: UnrealBuildTool.dll not found in %UBTPath%
 pause
 goto Exit
 
diff --git a/Engine/Build/BatchFiles/GenerateProjectFiles.bat b/Engine/Build/BatchFiles/GenerateProjectFiles.bat
--- a/Engine/Build/BatchFiles/GenerateProjectFiles.bat
+++ b/Engine/Build/BatchFiles/GenerateProjectFiles.bat
@@ -58,7 +58,7 @@ if errorlevel 1 goto Error_UBTCompileFailed
 
 rem ## Run UnrealBuildTool to generate Visual Studio solution and project files
 rem ## NOTE: We also pass along any arguments to the GenerateProjectFiles.bat here
-..\Binaries\DotNET\UnrealBuildTool\UnrealBuildTool.exe -ProjectFiles %*
+dotnet ..\Binaries\DotNET\UnrealBuildTool\UnrealBuildTool.dll -ProjectFiles %*
 if errorlevel 1 goto Error_ProjectGenerationFailed
 
 rem ## Success!
