UnrealBuildTool.exe -> dotnet UnrealBuildTool.dll (2/n)

Various build scripts

diff --git a/Engine/Build/BatchFiles/Linux/Build.sh b/Engine/Build/BatchFiles/Linux/Build.sh
--- a/Engine/Build/BatchFiles/Linux/Build.sh
+++ b/Engine/Build/BatchFiles/Linux/Build.sh
@@ -17,6 +17,6 @@ if [ ! -f Engine/Build/InstalledBuild.txt ]; then
   fi
 fi
 
-echo Running command : Engine/Binaries/DotNET/UnrealBuildTool/UnrealBuildTool "$@"
-Engine/Binaries/DotNET/UnrealBuildTool/UnrealBuildTool "$@"
+echo Running command : dotnet Engine/Binaries/DotNET/UnrealBuildTool/UnrealBuildTool.dll "$@"
+dotnet Engine/Binaries/DotNET/UnrealBuildTool/UnrealBuildTool.dll "$@"
 exit $?
diff --git a/Engine/Build/BatchFiles/Linux/GenerateProjectFiles.sh b/Engine/Build/BatchFiles/Linux/GenerateProjectFiles.sh
--- a/Engine/Build/BatchFiles/Linux/GenerateProjectFiles.sh
+++ b/Engine/Build/BatchFiles/Linux/GenerateProjectFiles.sh
@@ -31,4 +31,4 @@ if [ -f "$BASE_PATH/../../../Source/Programs/UnrealBuildTool/UnrealBuildTool.csp
 fi
 
 # pass all parameters to UBT
-"$BASE_PATH/../../../Binaries/DotNET/UnrealBuildTool/UnrealBuildTool" -projectfiles "$@"
\ No newline at end of file
+dotnet "$BASE_PATH/../../../Binaries/DotNET/UnrealBuildTool/UnrealBuildTool.dll" -projectfiles "$@"
\ No newline at end of file
diff --git a/Engine/Build/BatchFiles/Mac/Build.sh b/Engine/Build/BatchFiles/Mac/Build.sh
--- a/Engine/Build/BatchFiles/Mac/Build.sh
+++ b/Engine/Build/BatchFiles/Mac/Build.sh
@@ -19,14 +19,14 @@ if [ ! -f Engine/Build/InstalledBuild.txt ]; then
     for i in "$@" ; do
 	if [[ $i == "-buildscw" ]] ; then
 		echo Building ShaderCompileWorker...
-		Engine/Binaries/DotNET/UnrealBuildTool/UnrealBuildTool ShaderCompileWorker Mac Development
+		dotnet Engine/Binaries/DotNET/UnrealBuildTool/UnrealBuildTool.dll ShaderCompileWorker Mac Development
 		break
 	fi
     done
 fi
 
-echo Running Engine/Binaries/DotNET/UnrealBuildTool/UnrealBuildTool "$@"
-Engine/Binaries/DotNET/UnrealBuildTool/UnrealBuildTool "$@"
+echo Running dotnet Engine/Binaries/DotNET/UnrealBuildTool/UnrealBuildTool.dll "$@"
+dotnet Engine/Binaries/DotNET/UnrealBuildTool/UnrealBuildTool.dll "$@"
 
 ExitCode=$?
 if [ $ExitCode -eq 254 ] || [ $ExitCode -eq 255 ] || [ $ExitCode -eq 2 ]; then
diff --git a/Engine/Build/BatchFiles/Mac/GenerateProjectFiles.sh b/Engine/Build/BatchFiles/Mac/GenerateProjectFiles.sh
--- a/Engine/Build/BatchFiles/Mac/GenerateProjectFiles.sh
+++ b/Engine/Build/BatchFiles/Mac/GenerateProjectFiles.sh
@@ -31,4 +31,4 @@ if [ -f "$BASE_PATH/../../../Source/Programs/UnrealBuildTool/UnrealBuildTool.csp
 fi
 
 # pass all parameters to UBT
-"$BASE_PATH/../../../Binaries/DotNET/UnrealBuildTool/UnrealBuildTool" -projectfiles "$@"
+dotnet "$BASE_PATH/../../../Binaries/DotNET/UnrealBuildTool/UnrealBuildTool.dll" -projectfiles "$@"
diff --git a/Engine/Build/BatchFiles/Mac/XcodeBuild.sh b/Engine/Build/BatchFiles/Mac/XcodeBuild.sh
--- a/Engine/Build/BatchFiles/Mac/XcodeBuild.sh
+++ b/Engine/Build/BatchFiles/Mac/XcodeBuild.sh
@@ -103,15 +103,15 @@ if [ "$ACTION" == "build" ]; then
 
 	# Build SCW if this is an editor target
 	if [[ "$TARGET" == *"Editor" ]]; then
-		Engine/Binaries/DotNET/UnrealBuildTool/UnrealBuildTool ShaderCompileWorker Mac Development
+		dotnet Engine/Binaries/DotNET/UnrealBuildTool/UnrealBuildTool.dll ShaderCompileWorker Mac Development
 	fi
 
 elif [ $ACTION == "clean" ]; then
 	AdditionalFlags="-clean"
 fi
 
-echo Running Engine/Binaries/DotNET/UnrealBuildTool/UnrealBuildTool $TARGET $PLATFORM $CONFIGURATION "$TRAILINGARGS" $UBT_ARCHFLAG $AdditionalFlags
-Engine/Binaries/DotNET/UnrealBuildTool/UnrealBuildTool $TARGET $PLATFORM $CONFIGURATION "$TRAILINGARGS" $UBT_ARCHFLAG $AdditionalFlags
+echo Running dotnet Engine/Binaries/DotNET/UnrealBuildTool/UnrealBuildTool.dll $TARGET $PLATFORM $CONFIGURATION "$TRAILINGARGS" $UBT_ARCHFLAG $AdditionalFlags
+dotnet Engine/Binaries/DotNET/UnrealBuildTool/UnrealBuildTool.dll $TARGET $PLATFORM $CONFIGURATION "$TRAILINGARGS" $UBT_ARCHFLAG $AdditionalFlags
 
 ExitCode=$?
 if [ $ExitCode -eq 254 ] || [ $ExitCode -eq 255 ] || [ $ExitCode -eq 2 ]; then
diff --git a/Engine/Build/BatchFiles/RunUBT.sh b/Engine/Build/BatchFiles/RunUBT.sh
--- a/Engine/Build/BatchFiles/RunUBT.sh
+++ b/Engine/Build/BatchFiles/RunUBT.sh
@@ -36,4 +36,4 @@ if [ -f "$BASE_PATH/../../Source/Programs/UnrealBuildTool/UnrealBuildTool.csproj
 fi
 
 # pass all parameters to UBT
-"$BASE_PATH/../../Binaries/DotNET/UnrealBuildTool/UnrealBuildTool" "$@"
+dotnet "$BASE_PATH/../../Binaries/DotNET/UnrealBuildTool/UnrealBuildTool.dll" "$@"
