Use installed UBT when run from an installed build

diff --git a/Engine/Build/BatchFiles/Mac/GenerateProjectFiles.sh b/Engine/Build/BatchFiles/Mac/GenerateProjectFiles.sh
--- a/Engine/Build/BatchFiles/Mac/GenerateProjectFiles.sh
+++ b/Engine/Build/BatchFiles/Mac/GenerateProjectFiles.sh
@@ -21,7 +21,8 @@ if [ ! -d "$BASE_PATH/../../../Source" ]; then
 fi
 
 source "$BASE_PATH/SetupEnvironment.sh" -dotnet "$BASE_PATH"
-if [ -f "$BASE_PATH/../../../Source/Programs/UnrealBuildTool/UnrealBuildTool.csproj" ]; then
+# ensure UnrealBuildTool is up to date if the project file exists, but not if running from an installed build
+if [ -f "$BASE_PATH/../../../Source/Programs/UnrealBuildTool/UnrealBuildTool.csproj" -a ! -f "$BASE_PATH/../../../Build/InstalledBuild.txt"]; then
   dotnet msbuild /restore /target:build /property:Configuration=Development /nologo "$BASE_PATH/../../../Source/Programs/UnrealBuildTool/UnrealBuildTool.csproj" /verbosity:quiet
 
   if [ $? -ne 0 ]; then
