Make killing the children of RunUAT more robust

diff --git a/Engine/Build/BatchFiles/RunUAT.sh b/Engine/Build/BatchFiles/RunUAT.sh
--- a/Engine/Build/BatchFiles/RunUAT.sh
+++ b/Engine/Build/BatchFiles/RunUAT.sh
@@ -10,9 +10,41 @@ echo
 echo Running AutomationTool...
 echo
 
+GetAllChildProcesses() {
+	local Children=$(ps -o pid= --ppid "$1")
+
+	for PID in $Children
+	do
+		GetAllChildProcesses "$PID"
+	done
+
+	echo "$Children"
+}
+
+# Gather all the descendant children of this process, and first kill -TERM. If any child process
+# is still alive finally send a -KILL
 TermHandler() {
-	GPID=$(ps -o pgid= $UATPid)
-	kill -TERM -$GPID 2> /dev/null
+	MaxWait=30
+	CurrentWait=0
+
+	ProcessesToKill=$(GetAllChildProcesses $$)
+	kill -s TERM $ProcessesToKill 2> /dev/null
+
+	ProcessesStillAllive=$(ps -o pid= -p $ProcessesToKill)
+
+	# Wait until all the processes have been gracefully killed, or max Wait time
+	while [ -n "$ProcessesStillAllive" ] && [ "$CurrentWait" -lt "$MaxWait" ]
+	do
+		CurrentWait=$((CurrentWait + 1))
+		sleep 1
+
+		ProcessesStillAllive=$(ps -o pid= -p $ProcessesToKill)
+	done
+
+	# If some processes are still alive after MaxWait, lets just force kill them
+	if [ -n "$ProcessesStillAllive" ]; then
+		kill -s KILL $ProcessesStillAllive 2> /dev/null
+	fi
 }
 
 # loop over the arguments, quoting spaces to pass to UAT proper
@@ -109,7 +141,6 @@ else
 	LogDir="$uebp_LogFolder"
 fi
 
-
 # if we are running under UE, we need to run this with the term handler (otherwise canceling a UAT job from the editor
 # can leave mono, etc running in the background, which means we need the PID so we 
 # run it in the background
