Fix Mac GenerateProjectFiles.sh argument quoting

diff --git a/Engine/Build/BatchFiles/Mac/GenerateProjectFiles.sh b/Engine/Build/BatchFiles/Mac/GenerateProjectFiles.sh
--- a/Engine/Build/BatchFiles/Mac/GenerateProjectFiles.sh
+++ b/Engine/Build/BatchFiles/Mac/GenerateProjectFiles.sh
@@ -31,16 +31,5 @@ if [ -f "$BASE_PATH/../../../Source/Programs/UnrealBuildTool/UnrealBuildTool.csp
   fi
 fi
 
-if [ $# -lt 1 ]; then
-	# Create an empty string argument for UBT; otherwise we see an index out-of-range error.
-	QuotedArgs="\"\""
-else
-	# Apply quote-escape to all arguments
-	QuotedArgs="$(printf " %q" "$@")"
-
-	# Trim leading whitespace
-	QuotedArgs="${QuotedArgs#"${QuotedArgs%%[![:space:]]*}"}"
-fi
-
 # pass all parameters to UBT
-dotnet "$BASE_PATH/../../../Binaries/DotNET/UnrealBuildTool/UnrealBuildTool.dll" -projectfiles "${QuotedArgs}"
+"$BASE_PATH/../../../Binaries/DotNET/UnrealBuildTool/UnrealBuildTool" -projectfiles "$@"
