SetupDotnet.sh:

Handle different major version for installed dotnet

diff --git a/Engine/Build/BatchFiles/Linux/SetupDotnet.sh b/Engine/Build/BatchFiles/Linux/SetupDotnet.sh
--- a/Engine/Build/BatchFiles/Linux/SetupDotnet.sh
+++ b/Engine/Build/BatchFiles/Linux/SetupDotnet.sh
@@ -12,7 +12,11 @@ if [ "$UE_USE_SYSTEM_DOTNET" == "1" ] && [ ! $DOTNET_VERSION_PATH == "" ] && [ -
 	DOTNET_SDKS=(`dotnet --list-sdks | grep -P "(\d*)\.(\d*)\..* \[(.*)\]"`)
 	for DOTNET_SDK in $DOTNET_SDKS
     do
-        if [ ${DOTNET_SDK[0]} -ge 3 ]; then
+		if [ ${DOTNET_SDK[0]} -gt 3 ]; then
+			IS_DOTNET_INSTALLED=1
+		fi
+
+        if [ ${DOTNET_SDK[0]} -eq 3 ]; then
             if [ ${DOTNET_SDK[1]} -ge 1 ]; then
                 IS_DOTNET_INSTALLED=1
             fi
diff --git a/Engine/Build/BatchFiles/Mac/SetupDotnet.sh b/Engine/Build/BatchFiles/Mac/SetupDotnet.sh
--- a/Engine/Build/BatchFiles/Mac/SetupDotnet.sh
+++ b/Engine/Build/BatchFiles/Mac/SetupDotnet.sh
@@ -12,7 +12,11 @@ if [ "$UE_USE_SYSTEM_DOTNET" == "1" ] && [ ! $DOTNET_VERSION_PATH == "" ] && [ -
 	DOTNET_SDKS=(`dotnet --list-sdks | grep -P "(\d*)\.(\d*)\..* \[(.*)\]"`)
 	for DOTNET_SDK in $DOTNET_SDKS
 	do
-		if [ ${DOTNET_SDK[0]} -ge 3 ]; then
+		if [ ${DOTNET_SDK[0]} -gt 3 ]; then
+			IS_DOTNET_INSTALLED=1
+		fi
+
+		if [ ${DOTNET_SDK[0]} -eq 3 ]; then
 			if [ ${DOTNET_SDK[1]} -ge 1 ]; then
 				IS_DOTNET_INSTALLED=1
 			fi
