Changing RunUBT.bat to invoke UnrealBuildTool correctly.

diff --git a/Engine/Build/BatchFiles/RunUBT.bat b/Engine/Build/BatchFiles/RunUBT.bat
--- a/Engine/Build/BatchFiles/RunUBT.bat
+++ b/Engine/Build/BatchFiles/RunUBT.bat
@@ -49,7 +49,7 @@ if errorlevel 1 goto Error_UBTCompileFailed
 
 rem ## Run UnrealBuildTool to generate Visual Studio solution and project files
 rem ## NOTE: We also pass along any arguments here
-..\Binaries\DotNET\UnrealBuildTool\UnrealBuildTool.exe %*
+dotnet ..\Binaries\DotNET\UnrealBuildTool\UnrealBuildTool.dll %*
 if errorlevel 1 goto Error_UBTRunFailed
 
 rem ## Success!
