Fix GitDependencies.sh creating invalid symlink.

diff --git a/Engine/Build/BatchFiles/Mac/GitDependencies.sh b/Engine/Build/BatchFiles/Mac/GitDependencies.sh
--- a/Engine/Build/BatchFiles/Mac/GitDependencies.sh
+++ b/Engine/Build/BatchFiles/Mac/GitDependencies.sh
@@ -19,10 +19,6 @@ source SetupEnvironment.sh -mono "`pwd`"
 
 cd ../../../..
 
-if [ ! -f Engine/Binaries/ThirdParty/Mono/Mac/lib/libmsvcrt.dylib ]; then
-	ln -s /usr/lib/libc.dylib Engine/Binaries/ThirdParty/Mono/Mac/lib/libmsvcrt.dylib
-fi
-
 mono Engine/Binaries/DotNET/GitDependencies.exe "$@"
 
 pushd "$(dirname "$SCRIPT_PATH")" > /dev/null
