IterationProfile: Capture new performance baselines for new horde resources

The new resource pool is 3x Xeon 6128 3.4Ghz w/ 64GB (RDU-R7920-*)

These were generated with the following command line:
  Engine\Build\BatchFiles\IterationProfileBaseline.bat -set:BaselineBuildName="++UE5+Main-CL-16768700"

diff --git a/Engine/Build/BatchFiles/IterationProfileBaseline.bat b/Engine/Build/BatchFiles/IterationProfileBaseline.bat
--- a/Engine/Build/BatchFiles/IterationProfileBaseline.bat
+++ b/Engine/Build/BatchFiles/IterationProfileBaseline.bat
@@ -1 +1 @@
-@call %~dp0RunUAT.bat BuildGraph -Script=Engine/Build/IterationProfile.xml -Target="Baseline Iteration Profile Results" -append:ProjectNames="ShooterGame" -set:Win64=true %*
+@call %~dp0RunUAT.bat BuildGraph -Script=Engine/Build/IterationProfile.xml -Target="Baseline Iteration Profile Results" -append:ProjectNames="EngineTest" -append:ProjectNames="Reverb" -append:ProjectNames="ShooterGame" -set:Win64=true %*
