IterationProfile improvements

Fix baseline script for proper phases
Fix horde labels to use run nodes instead of agent tags

diff --git a/Engine/Build/BatchFiles/IterationProfileBaseline.bat b/Engine/Build/BatchFiles/IterationProfileBaseline.bat
--- a/Engine/Build/BatchFiles/IterationProfileBaseline.bat
+++ b/Engine/Build/BatchFiles/IterationProfileBaseline.bat
@@ -1 +1 @@
-@call %~dp0RunUAT.bat BuildGraph -Script=Engine/Build/IterationProfile.xml -Target="Baseline Iteration Profile Results" -set:Platforms=Win64 -append:Projects="EngineTest" -append:Projects="Reverb" -append:Projects="ShooterGame" %*
+@call %~dp0RunUAT.bat BuildGraph -Script=Engine/Build/IterationProfile.xml -Target="Baseline Iteration Profile Results" -set:Platforms=Win64 -set:Phases="Cold;Warm" -append:Projects="EngineTest" -append:Projects="Reverb" -append:Projects="ShooterGame" %*
diff --git a/Engine/Build/IterationProfile.xml b/Engine/Build/IterationProfile.xml
--- a/Engine/Build/IterationProfile.xml
+++ b/Engine/Build/IterationProfile.xml
@@ -174,11 +174,10 @@
 					</ForEach>
 				</Node>
 				<Property Name="BaselineGraphNodes" Value="$(BaselineGraphNodes);$(BaselineNodeName)"/>
+
+				<Label Category="Test" Name="$(Project) Iteration Profile $(Platform)" Requires="$(RunNodeName)"/>
 			</Agent>
 
-			<!-- Graph nodes that really tie the room together -->
-			<Label Category="Test" Name="$(Project) Iteration Profile $(Platform)" Requires="Prerequisites Agent $(Platform);Automated Test Agent $(Project) $(Platform)"/>
-			
 		</ForEach> <!-- Projects -->
 
 	</ForEach> <!-- Platforms -->
