GenerateProjectFiles.sh

Use UnrealBuildTool dll, rather than exe

#trivial

diff --git a/Engine/Build/BatchFiles/Mac/GenerateProjectFiles.sh b/Engine/Build/BatchFiles/Mac/GenerateProjectFiles.sh
--- a/Engine/Build/BatchFiles/Mac/GenerateProjectFiles.sh
+++ b/Engine/Build/BatchFiles/Mac/GenerateProjectFiles.sh
@@ -32,4 +32,4 @@ if [ -f "$BASE_PATH/../../../Source/Programs/UnrealBuildTool/UnrealBuildTool.csp
 fi
 
 # pass all parameters to UBT
-"$BASE_PATH/../../../Binaries/DotNET/UnrealBuildTool/UnrealBuildTool" -projectfiles "$@"
+dotnet "$BASE_PATH/../../../Binaries/DotNET/UnrealBuildTool/UnrealBuildTool.dll" -projectfiles "$@"
