IterationProfile: Replace Reverb with AncientGame

diff --git a/Engine/Build/BatchFiles/IterationProfile.bat b/Engine/Build/BatchFiles/IterationProfile.bat
--- a/Engine/Build/BatchFiles/IterationProfile.bat
+++ b/Engine/Build/BatchFiles/IterationProfile.bat
@@ -1 +1 @@
-@call %~dp0RunUAT.bat BuildGraph -Script=Engine/Build/IterationProfile.xml -Target="Run Iteration Profile Tests" -set:Platforms=Win64 -append:Projects="EngineTest" -append:Projects="Reverb" -append:Projects="ShooterGame" %*
+@call %~dp0RunUAT.bat BuildGraph -Script=Engine/Build/IterationProfile.xml -Target="Run Iteration Profile Tests" -set:Platforms=Win64 -set:Phases="Cold;Warm" -append:Projects="AncientGame" -append:Projects="EngineTest" -append:Projects="ShooterGame" %*
diff --git a/Engine/Build/BatchFiles/IterationProfileBaseline.bat b/Engine/Build/BatchFiles/IterationProfileBaseline.bat
--- a/Engine/Build/BatchFiles/IterationProfileBaseline.bat
+++ b/Engine/Build/BatchFiles/IterationProfileBaseline.bat
@@ -1 +1 @@
-@call %~dp0RunUAT.bat BuildGraph -Script=Engine/Build/IterationProfile.xml -Target="Baseline Iteration Profile Results" -set:Platforms=Win64 -set:Phases="Cold;Warm" -append:Projects="EngineTest" -append:Projects="Reverb" -append:Projects="ShooterGame" %*
+@call %~dp0RunUAT.bat BuildGraph -Script=Engine/Build/IterationProfile.xml -Target="Baseline Iteration Profile Results" -set:Platforms=Win64 -set:Phases="Cold;Warm" -append:Projects="AncientGame" -append:Projects="EngineTest" -append:Projects="ShooterGame" %*
