Add support for LLDB pretty printer for Linux

diff --git a/Engine/Build/BatchFiles/Linux/GenerateLLDBInit.sh b/Engine/Build/BatchFiles/Linux/GenerateLLDBInit.sh
new file mode 100755
--- /dev/null
+++ b/Engine/Build/BatchFiles/Linux/GenerateLLDBInit.sh
@@ -0,0 +1,6 @@
+#!/bin/bash
+
+if [ ! -f ~/.lldbinit ]; then
+  echo -ne "settings set target.inline-breakpoint-strategy always\n" > ~/.lldbinit
+  echo -ne "command script import \"`pwd`/../../../Extras/LLDBDataFormatters/UEDataFormatters_2ByteChars.py\"\n" >> ~/.lldbinit
+fi
diff --git a/GenerateProjectFiles.sh b/GenerateProjectFiles.sh
--- a/GenerateProjectFiles.sh
+++ b/GenerateProjectFiles.sh
@@ -25,6 +25,7 @@ if [ "$(uname)" = "Darwin" ]; then
 else
 	# assume (GNU/)Linux
 	cd Engine/Build/BatchFiles/Linux
+	bash ./GenerateLLDBInit.sh
 	bash ./GenerateGDBInit.sh
 	bash ./GenerateProjectFiles.sh "$@"
 fi
