Fixed Turnkey not being able to find batch file to update
environment variables
#rnx

diff --git a/Engine/Build/BatchFiles/RunUAT.bat b/Engine/Build/BatchFiles/RunUAT.bat
--- a/Engine/Build/BatchFiles/RunUAT.bat
+++ b/Engine/Build/BatchFiles/RunUAT.bat
@@ -83,8 +83,8 @@ if EXIST %SCRIPT_DIR%..\..\Intermediate\Turnkey\PostTurnkeyVariables.bat (
 	rem ## We need to endlocal so that the vars in the batch file work. NOTE: Working directory from pushd will be UNDONE here, but since we are about to quit, it's okay
 	endlocal 
 	echo Updating environment variables set by a Turnkey sub-process
-	call %SCRIPT_DIR%..\..\Engine\Intermediate\Turnkey\PostTurnkeyVariables.bat
-	del %SCRIPT_DIR%..\..\Engine\Intermediate\Turnkey\PostTurnkeyVariables.bat
+	call %SCRIPT_DIR%..\..\Intermediate\Turnkey\PostTurnkeyVariables.bat
+	del %SCRIPT_DIR%..\..\Intermediate\Turnkey\PostTurnkeyVariables.bat
 	rem ## setlocal again so that any popd's etc don't have an effect on calling process
 	setlocal
 )
