BuildUAT.sh:

Use redirection syntax with fc, rather than errorlevel (per https://ss64.com/nt/fc.html) to work around inconsistent results.

diff --git a/Engine/Build/BatchFiles/BuildUAT.bat b/Engine/Build/BatchFiles/BuildUAT.bat
--- a/Engine/Build/BatchFiles/BuildUAT.bat
+++ b/Engine/Build/BatchFiles/BuildUAT.bat
@@ -43,8 +43,8 @@ if /I "%ARGUMENT%" == "FORCE" goto Build_AutomationTool
 
 :Check_UpToDate
 set RUNUAT_EXITCODE=0
-fc /b ..\Intermediate\Build\AutomationToolFiles.txt ..\Intermediate\Build\AutomationToolPrevFiles.txt  >nul 2>nul
-if not errorlevel 1 goto Exit
+rem per https://ss64.com/nt/fc.html using redirection syntax rather than errorlevel, based on observed inconsistent results from this function
+fc ..\Intermediate\Build\AutomationToolFiles.txt ..\Intermediate\Build\AutomationToolPrevFiles.txt >nul && goto Exit
 
 :Build_AutomationTool
 rem ## Verify that dotnet is present
