RunUAT.bat: Always build AutomationTool with msbuild

Avoid problem where insufficient dependency checking would fail to rebuild AutomationTool (and dependent projects) when needed.

diff --git a/Engine/Build/BatchFiles/RunUAT.bat b/Engine/Build/BatchFiles/RunUAT.bat
--- a/Engine/Build/BatchFiles/RunUAT.bat
+++ b/Engine/Build/BatchFiles/RunUAT.bat
@@ -52,6 +52,14 @@ rem ## Verify that dotnet is present
 call "%SCRIPT_DIR%GetDotnetPath.bat"
 if errorlevel 1 goto Error_NoDotnetSDK
 
+
+
+rem Checking for out-of-date files won't find source files in places like Engine/Platform, resulting in occasionally
+rem out of date builds. Until a better solution is found, always try to build AutomationTool.
+set FORCECOMPILE_UAT=FORCE
+
+
+
 call "%SCRIPT_DIR%BuildUAT.bat" %FORCECOMPILE_UAT%
 if errorlevel 1 goto Error_UATCompileFailed
 
