Remove old SSH key directory when generating a new one, to fix up any bad directory permissions

diff --git a/Engine/Build/BatchFiles/MakeAndInstallSSHKey.bat b/Engine/Build/BatchFiles/MakeAndInstallSSHKey.bat
--- a/Engine/Build/BatchFiles/MakeAndInstallSSHKey.bat
+++ b/Engine/Build/BatchFiles/MakeAndInstallSSHKey.bat
@@ -49,7 +49,7 @@ echo ===========================================================================
 echo.
 
 pause
-
+rmdir "%KEY_DIR%" /s/q
 mkdir "%KEY_DIR%" 2> NUL
 pushd %SSHDIR%
 %RSYNC% -az -e '%SSH% -p %SSHPORT%' "%USER%@%MACHINE%":.ssh/RemoteToolChainPrivate.key "%CYGWIN_KEY_PATH%"
