Update scripts for next toolchain build

diff --git a/Engine/Build/BatchFiles/Linux/Toolchain/BuildLibCxx.sh b/Engine/Build/BatchFiles/Linux/Toolchain/BuildLibCxx.sh
--- a/Engine/Build/BatchFiles/Linux/Toolchain/BuildLibCxx.sh
+++ b/Engine/Build/BatchFiles/Linux/Toolchain/BuildLibCxx.sh
@@ -66,7 +66,7 @@ BuildLibCxx()
     set +x
 
     make -j ${CORES} cxxabi cxx
-    make install-libcxx-headers
+    make install-cxx-headers
 
     mkdir -p ${INSTALL_DIR}/lib/Linux/${ARCH}
     cp lib/libc++.a ${INSTALL_DIR}/lib/Linux/${ARCH}
diff --git a/Engine/Build/BatchFiles/Linux/Toolchain/build_linux_toolchain.sh b/Engine/Build/BatchFiles/Linux/Toolchain/build_linux_toolchain.sh
--- a/Engine/Build/BatchFiles/Linux/Toolchain/build_linux_toolchain.sh
+++ b/Engine/Build/BatchFiles/Linux/Toolchain/build_linux_toolchain.sh
@@ -3,8 +3,8 @@
 set -x
 set -eu
 
-ToolChainVersion=v17
-LLVM_VERSION=10.0.1
+ToolChainVersion=v18
+LLVM_VERSION=11.0.1
 LLVM_URL=[REDACTED_PROFILE_URL]/llvm-project/releases/download/llvmorg-${LLVM_VERSION}
 
 ToolChainVersionName="${ToolChainVersion}_clang-${LLVM_VERSION}-centos7"
