Update window toolchain building scripts for next time

diff --git a/Engine/Build/BatchFiles/Linux/Toolchain/RunMe.bat b/Engine/Build/BatchFiles/Linux/Toolchain/RunMe.bat
--- a/Engine/Build/BatchFiles/Linux/Toolchain/RunMe.bat
+++ b/Engine/Build/BatchFiles/Linux/Toolchain/RunMe.bat
@@ -1,10 +1,10 @@
 @REM @echo off
 
-set TOOLCHAIN_VERSION=v17
-set LLVM_VERSION=10.0.1
+set TOOLCHAIN_VERSION=v18
+set LLVM_VERSION=11.0.1
 
-set SVN_BINARY=%CD%\..\..\..\..\..\..\Binaries\ThirdParty\svn\Win64\svn.exe
-set CMAKE_BINARY=%CD%\..\..\..\..\..\..\Extras\ThirdPartyNotUE\CMake\bin\cmake.exe
+set SVN_BINARY=%CD%\..\..\..\..\Binaries\ThirdParty\svn\Win64\svn.exe
+set CMAKE_BINARY=%CD%\..\..\..\..\Extras\ThirdPartyNotUE\CMake\bin\cmake.exe
 set PYTHON_BINARY=%CD%\..\..\..\..\Binaries\ThirdParty\Python\Win64\python.exe
 set NSIS_BINARY=C:\Program Files (x86)\NSIS\Bin\makensis.exe
 
@@ -27,10 +27,11 @@ set ROOT_DIR=%CD%
 rm -rf %TEMP:\=/%\clang-build-%LLVM_VERSION%
 mkdir %TEMP%\clang-build-%LLVM_VERSION%
 pushd %TEMP%\clang-build-%LLVM_VERSION%
+
 
 unzip -o %ROOT_DIR:\=/%/%FILENAME%-windows.zip -d OUTPUT
 
-set GIT_LLVM_RELEASE_HASH=ef32c611aa214dea855364efd7ba451ec5ec3f74
+set GIT_LLVM_RELEASE_HASH=43ff75f2c3feef64f9d73328230d34dac8832a91
 
 git clone [REDACTED_PROFILE_URL]/llvm-project source
 pushd source
@@ -47,7 +48,7 @@ popd
 
 mkdir build_lld
 pushd build_lld
-
+
 %CMAKE_BINARY% -G "Visual Studio 16 2019" -DLLVM_TEMPORARILY_ALLOW_OLD_TOOLCHAIN=ON -DCMAKE_INSTALL_PREFIX="..\install" -DPYTHON_EXECUTABLE="%PYTHON_BINARY%" -DLLVM_CONFIG_PATH="..\install\bin\llvm-config.exe" "..\source\lld"
 %CMAKE_BINARY% -G "Visual Studio 16 2019" -DLLVM_TEMPORARILY_ALLOW_OLD_TOOLCHAIN=ON -DCMAKE_INSTALL_PREFIX="..\install" -DPYTHON_EXECUTABLE="%PYTHON_BINARY%" -DLLVM_CONFIG_PATH="..\install\bin\llvm-config.exe" "..\source\lld"
 %CMAKE_BINARY% --build . --target install --config MinSizeRel
