RunUAT.sh: Always build AutomationTool with msbuild
Avoid problem where insufficient dependency checking would fail to rebuild AutomationTool (and dependent projects) when needed.

 - Implement RunUAT.bat changes in CL 17505693 for RunUAT.sh

#rnx

diff --git a/Engine/Build/BatchFiles/RunUAT.sh b/Engine/Build/BatchFiles/RunUAT.sh
--- a/Engine/Build/BatchFiles/RunUAT.sh
+++ b/Engine/Build/BatchFiles/RunUAT.sh
@@ -91,13 +91,17 @@ if [ "$(uname)" = "Linux" ]; then
 	source "$SCRIPT_DIR/Linux/SetupEnvironment.sh" $EnvironmentType "$SCRIPT_DIR/Linux"
 fi
 
+# Checking for out-of-date files won't find source files in places like Engine/Platform, resulting in occasionally
+# out of date builds. Until a better solution is found, always try to build AutomationTool.
+FORCECOMPILE_UAT=FORCE
+
 if [ "$UATCompileArg" = "-compile" ]; then
   # see if the .csproj exists to be compiled
 	if [ ! -f Source/Programs/AutomationTool/AutomationTool.csproj ]; then
 		echo No project to compile, attempting to use precompiled AutomationTool
 		UATCompileArg=
 	else
-		"$SCRIPT_DIR"/BuildUAT.sh $MSBuild_Verbosity
+		"$SCRIPT_DIR"/BuildUAT.sh $MSBuild_Verbosity $FORCECOMPILE_UAT
 		if [ $? -ne 0 ]; then
 			echo RunUAT ERROR: AutomationTool failed to compile.
 			exit 1
