PR #8639: Fix VS cleaning not working on installed builds (Contributed by [REDACTED_PERSON])

diff --git a/Engine/Build/BatchFiles/Clean.bat b/Engine/Build/BatchFiles/Clean.bat
--- a/Engine/Build/BatchFiles/Clean.bat
+++ b/Engine/Build/BatchFiles/Clean.bat
@@ -1,75 +1,17 @@
-@echo off
-
-rem ## Unreal Engine 4 cleanup script
-rem ## Copyright Epic Games, Inc. All Rights Reserved.
-rem ##
-rem ## This script is expecting to exist in the UE4 root directory.  It will not work correctly
-rem ## if you copy it to a different location and run it.
-rem ##
-rem ##     %1 is the game name
-rem ##     %2 is the platform name
-rem ##     %3 is the configuration name
-
-setlocal
-
-rem ## First, make sure the batch file exists in the folder we expect it to.  This is necessary in order to
-rem ## verify that our relative path to the /Engine/Source directory is correct
-if not exist "%~dp0..\..\Source" goto Error_BatchFileInWrongLocation
-
-
-rem ## Change the CWD to /Engine/Source.  We always need to run UnrealBuildTool from /Engine/Source!
-pushd "%~dp0..\..\Source"
-if not exist ..\Build\BatchFiles\Clean.bat goto Error_BatchFileInWrongLocation
-
-rem ## If this is an installed build, we don't need to rebuild UBT. Go straight to cleaning.
-if exist ..\Build\InstalledBuild.txt goto ReadyToClean
-
-rem ## Get the path to MSBuild
-call "%~dp0GetMSBuildPath.bat"
-if errorlevel 1 goto Error_NoVisualStudioEnvironment
-
-
-rem ## Compile UBT if the project file exists
-
-set ProjectFile="Programs\UnrealBuildTool\UnrealBuildTool.csproj"
-set UBTPath="..\..\Engine\Binaries\DotNET\UnrealBuildTool\UnrealBuildTool.dll"
-if not exist %ProjectFile% goto NoProjectFile
-dotnet build Programs\UnrealBuildTool\UnrealBuildTool.csproj -c Development -v quiet
-if errorlevel 1 goto Error_UBTCompileFailed
-:NoProjectFile
-
-
-rem ## Execute UBT
-:ReadyToClean
-if not exist %UBTPath% goto Error_UBTMissing
-dotnet %UBTPath% %* -clean
-goto Exit
-
-
-:Error_BatchFileInWrongLocation
-echo Clean ERROR: The batch file does not appear to be located in the UE4/Engine/Build/BatchFiles directory.  This script must be run from within that directory.
-pause
-goto Exit
-
-:Error_NoVisualStudioEnvironment
-echo GenerateProjectFiles ERROR: A valid version of Visual Studio does not appear to be installed.
-pause
-goto Exit
-
-:Error_UBTCompileFailed
-echo Clean ERROR: Failed to build UnrealBuildTool.
-pause
-goto Exit
-
-:Error_UBTMissing
-echo Clean ERROR: UnrealBuildTool.dll not found in %UBTPath%
-pause
-goto Exit
-
-:Error_CleanFailed
-echo Clean ERROR: Clean failed.
-pause
-goto Exit
-
-:Exit
-
+@echo off
+rem ## Unreal Engine 4 cleanup scriptrem ## Copyright Epic Games, Inc. All Rights Reserved.rem ##rem ## This script is expecting to exist in the UE4 root directory.  It will not work correctlyrem ## if you copy it to a different location and run it.rem ##rem ##     %1 is the game namerem ##     %2 is the platform namerem ##     %3 is the configuration name
+setlocal
+rem ## First, make sure the batch file exists in the folder we expect it to.  This is necessary in order torem ## verify that our relative path to the /Engine/Source directory is correctif not exist "%~dp0..\..\Source" goto Error_BatchFileInWrongLocation
+rem ## Change the CWD to /Engine/Source.  We always need to run UnrealBuildTool from /Engine/Source!pushd "%~dp0..\..\Source"if not exist ..\Build\BatchFiles\Clean.bat goto Error_BatchFileInWrongLocation
+set UBTPath="..\..\Engine\Binaries\DotNET\UnrealBuildTool\UnrealBuildTool.dll"
+rem ## If this is an installed build, we don't need to rebuild UBT. Go straight to cleaning.if exist ..\Build\InstalledBuild.txt goto ReadyToClean
+rem ## Get the path to MSBuildcall "%~dp0GetMSBuildPath.bat"if errorlevel 1 goto Error_NoVisualStudioEnvironment
+rem ## Compile UBT if the project file exists
+set ProjectFile="Programs\UnrealBuildTool\UnrealBuildTool.csproj"if not exist %ProjectFile% goto NoProjectFiledotnet build Programs\UnrealBuildTool\UnrealBuildTool.csproj -c Development -v quietif errorlevel 1 goto Error_UBTCompileFailed:NoProjectFile
+rem ## Execute UBT:ReadyToCleanif not exist %UBTPath% goto Error_UBTMissingdotnet %UBTPath% %* -cleangoto Exit
+:Error_BatchFileInWrongLocationecho Clean ERROR: The batch file does not appear to be located in the UE4/Engine/Build/BatchFiles directory.  This script must be run from within that directory.pausegoto Exit
+:Error_NoVisualStudioEnvironmentecho GenerateProjectFiles ERROR: A valid version of Visual Studio does not appear to be installed.pausegoto Exit
+:Error_UBTCompileFailedecho Clean ERROR: Failed to build UnrealBuildTool.pausegoto Exit
+:Error_UBTMissingecho Clean ERROR: UnrealBuildTool.dll not found in %UBTPath%pausegoto Exit
+:Error_CleanFailedecho Clean ERROR: Clean failed.pausegoto Exit
+:Exit
\ No newline at end of file
