Fix Jira UE-138247

Specifying the platform as AnyCPU explicitly when building UBT was necessary to trigger the conditional OutputPath property.

diff --git a/Engine/Build/BatchFiles/RunUBT.bat b/Engine/Build/BatchFiles/RunUBT.bat
--- a/Engine/Build/BatchFiles/RunUBT.bat
+++ b/Engine/Build/BatchFiles/RunUBT.bat
@@ -44,7 +44,7 @@ dotnet msbuild /target:clean /property:Configuration=Development /nologo Program
 
 :SkipClean
 echo Building UnrealBuildTool...
-dotnet msbuild /restore /target:build /property:Configuration=Development /nologo Programs\UnrealBuildTool\UnrealBuildTool.csproj /verbosity:quiet
+dotnet msbuild /restore /target:build /property:Configuration=Development /property:Platform=AnyCPU /nologo Programs\UnrealBuildTool\UnrealBuildTool.csproj /verbosity:quiet
 if errorlevel 1 goto Error_UBTCompileFailed
 
 rem ## Run UnrealBuildTool to generate Visual Studio solution and project files
@@ -97,4 +97,4 @@ goto Exit
 :Exit
 rem ## Restore original CWD in case we change it
 popd
-exit /B %RUNUBT_EXITCODE%
\ No newline at end of file
+exit /B %RUNUBT_EXITCODE%
