Force zlib enabled when building the toolchain, requires correctly static linking

diff --git a/Engine/Build/BatchFiles/Linux/Toolchain/build_linux_toolchain.sh b/Engine/Build/BatchFiles/Linux/Toolchain/build_linux_toolchain.sh
--- a/Engine/Build/BatchFiles/Linux/Toolchain/build_linux_toolchain.sh
+++ b/Engine/Build/BatchFiles/Linux/Toolchain/build_linux_toolchain.sh
@@ -167,6 +167,7 @@ for arch in $TARGETS; do
 			-DCMAKE_AR=${InstallClangDir}/bin/llvm-ar \
 			-DCMAKE_NM=${InstallClangDir}/bin/llvm-nm \
 			-DCMAKE_RANLIB=${InstallClangDir}/bin/llvm-ranlib \
+			-DLLVM_ENABLE_ZLIB=FORCE_ON
 			-DCMAKE_EXE_LINKER_FLAGS="--target=$arch -L${OutputDirLinux}/$arch/lib64 --sysroot=${OutputDirLinux}/$arch -fuse-ld=lld" \
 			-DCMAKE_C_FLAGS="--target=$arch --sysroot=${OutputDirLinux}/$arch" \
 			-DCMAKE_CXX_FLAGS="--target=$arch --sysroot=${OutputDirLinux}/$arch" \
