Fixed typos in wording of error statements and whitespace mistakes.

diff --git a/Engine/Build/BatchFiles/Linux/Build.sh b/Engine/Build/BatchFiles/Linux/Build.sh
--- a/Engine/Build/BatchFiles/Linux/Build.sh
+++ b/Engine/Build/BatchFiles/Linux/Build.sh
@@ -11,10 +11,10 @@ source Engine/Build/BatchFiles/Linux/SetupEnvironment.sh -dotnet Engine/Build/Ba
 
 # If this is a source drop of the engine make sure that the UnrealBuildTool is up-to-date
 if [ ! -f Engine/Build/InstalledBuild.txt ]; then
-    if ! dotnet build Engine/Source/Programs/UnrealBuildTool/UnrealBuildTool.csproj -c Development -v quiet; then
-    echo "Failed to build to build tool (UnrealBuildTool)"
-    exit 1
-  fi
+	if ! dotnet build Engine/Source/Programs/UnrealBuildTool/UnrealBuildTool.csproj -c Development -v quiet; then
+		echo "Failed to build the build tool (UnrealBuildTool)"
+		exit 1
+	fi
 fi
 
 echo Running command : dotnet Engine/Binaries/DotNET/UnrealBuildTool/UnrealBuildTool.dll "$@"
diff --git a/Engine/Build/BatchFiles/Mac/Build.sh b/Engine/Build/BatchFiles/Mac/Build.sh
--- a/Engine/Build/BatchFiles/Mac/Build.sh
+++ b/Engine/Build/BatchFiles/Mac/Build.sh
@@ -9,20 +9,20 @@ source Engine/Build/BatchFiles/Mac/SetupEnvironment.sh -dotnet Engine/Build/Batc
 
 # Skip UBT and SWC compile step if this is an installed build.
 if [ ! -f Engine/Build/InstalledBuild.txt ]; then
-    # First make sure that the UnrealBuildTool is up-to-date
-    if ! dotnet build Engine/Source/Programs/UnrealBuildTool/UnrealBuildTool.csproj -c Development -v quiet; then
-	  echo "Failed to build to build tool (UnrealBuildTool)"
-	  exit 1
-    fi
+	# First make sure that the UnrealBuildTool is up-to-date
+	if ! dotnet build Engine/Source/Programs/UnrealBuildTool/UnrealBuildTool.csproj -c Development -v quiet; then
+		echo "Failed to build the build tool (UnrealBuildTool)"
+		exit 1
+	fi
 
-    # build SCW if specified
-    for i in "$@" ; do
+	# build SCW if specified
+	for i in "$@" ; do
 	if [[ $i == "-buildscw" ]] ; then
 		echo Building ShaderCompileWorker...
 		dotnet Engine/Binaries/DotNET/UnrealBuildTool/UnrealBuildTool.dll ShaderCompileWorker Mac Development
 		break
 	fi
-    done
+	done
 fi
 
 echo Running dotnet Engine/Binaries/DotNET/UnrealBuildTool/UnrealBuildTool.dll "$@"
