Accepting PR fixing an error with escape sequences and recommending a different Linux build

diff --git a/Engine/Build/BatchFiles/Linux/GenerateProjectFiles.sh b/Engine/Build/BatchFiles/Linux/GenerateProjectFiles.sh
--- a/Engine/Build/BatchFiles/Linux/GenerateProjectFiles.sh
+++ b/Engine/Build/BatchFiles/Linux/GenerateProjectFiles.sh
@@ -11,24 +11,24 @@ BASE_PATH="`dirname "$0"`"
 # this is located inside an extra 'Linux' path unlike the Windows variant.
 
 if [ ! -d "$BASE_PATH/../../../Binaries/DotNET" ]; then
- echo GenerateProjectFiles ERROR: It looks like you're missing some files that are required in order to generate projects.  Please check that you've downloaded and unpacked the engine source code, binaries, content and third-party dependencies before running this script.
- exit 1
+	echo GenerateProjectFiles ERROR: It looks like you are missing some files that are required in order to generate projects.  Please check that you have downloaded and unpacked the engine source code, binaries, content and third-party dependencies before running this script.
+	exit 1
 fi
 
 if [ ! -d "$BASE_PATH/../../../Source" ]; then
- echo GenerateProjectFiles ERROR: This script file does not appear to be located inside the Engine/Build/BatchFiles/Mac directory.
- exit 1
+	echo GenerateProjectFiles ERROR: This script file does not appear to be located inside the Engine/Build/BatchFiles/Mac directory.
+	exit 1
 fi
 
 source "$BASE_PATH/SetupEnvironment.sh" -dotnet "$BASE_PATH"
 
 if [ -f "$BASE_PATH/../../../Source/Programs/UnrealBuildTool/UnrealBuildTool.csproj" ]; then
-  dotnet msbuild /restore /target:build /property:Configuration=Development /nologo $BASE_PATH/../../../Source/Programs/UnrealBuildTool/UnrealBuildTool.csproj /verbosity:quiet
-  if [ $? -ne 0 ]; then
-    echo GenerateProjectFiles ERROR: Failed to build UnrealBuildTool
-    exit 1
-  fi
+	dotnet msbuild /restore /target:build /property:Configuration=Development /nologo $BASE_PATH/../../../Source/Programs/UnrealBuildTool/UnrealBuildTool.csproj /verbosity:quiet
+	if [ $? -ne 0 ]; then
+		echo GenerateProjectFiles ERROR: Failed to build UnrealBuildTool
+		exit 1
+	fi
 fi
 
 # pass all parameters to UBT
-dotnet "$BASE_PATH/../../../Binaries/DotNET/UnrealBuildTool/UnrealBuildTool.dll" -projectfiles "$@"
\ No newline at end of file
+dotnet "$BASE_PATH/../../../Binaries/DotNET/UnrealBuildTool/UnrealBuildTool.dll" -projectfiles "$@"
diff --git a/Engine/Build/BatchFiles/Linux/README.md b/Engine/Build/BatchFiles/Linux/README.md
--- a/Engine/Build/BatchFiles/Linux/README.md
+++ b/Engine/Build/BatchFiles/Linux/README.md
@@ -47,8 +47,11 @@ How to set up the sources for building, step by step:
 
 1. Clone EpicGames/UnrealEngine repository
 
-    ``git clone https://github.com/EpicGames/UnrealEngine -b release``
-    
+	``git clone https://github.com/EpicGames/UnrealEngine -b <branch>``
+	i.g. ``git clone https://github.com/EpicGames/UnrealEngine -b ue5-main``
+	You can see the more explanation about how we use branch version in root directory UnrealEngine/README.md.
+	Early access 1 and 2 can be downloadable from Releases.
+
 2. Run Setup.sh once.
 
     ``cd UnrealEngine``
diff --git a/Engine/Build/BatchFiles/Mac/GenerateProjectFiles.sh b/Engine/Build/BatchFiles/Mac/GenerateProjectFiles.sh
--- a/Engine/Build/BatchFiles/Mac/GenerateProjectFiles.sh
+++ b/Engine/Build/BatchFiles/Mac/GenerateProjectFiles.sh
@@ -10,25 +10,25 @@ BASE_PATH="`dirname "$0"`"
 # this is located inside an extra 'Mac' path unlike the Windows variant.
 
 if [ ! -d "$BASE_PATH/../../../Binaries/DotNET" ]; then
- echo GenerateProjectFiles ERROR: It looks like you're missing some files that are required in order to generate projects.  Please check that you've downloaded and unpacked the engine source code, binaries, content and third-party dependencies before running this script.
- exit 1
+	echo GenerateProjectFiles ERROR: It looks like you are missing some files that are required in order to generate projects.  Please check that you have downloaded and unpacked the engine source code, binaries, content and third-party dependencies before running this script.
+	exit 1
 fi
 
 echo "$BASE_PATH"
 if [ ! -d "$BASE_PATH/../../../Source" ]; then
- echo GenerateProjectFiles ERROR: This script file does not appear to be located inside the Engine/Build/BatchFiles/Mac directory.
- exit 1
+	echo GenerateProjectFiles ERROR: This script file does not appear to be located inside the Engine/Build/BatchFiles/Mac directory.
+	exit 1
 fi
 
 source "$BASE_PATH/SetupEnvironment.sh" -dotnet "$BASE_PATH"
 # ensure UnrealBuildTool is up to date if the project file exists, but not if running from an installed build
 if [ -f "$BASE_PATH/../../../Source/Programs/UnrealBuildTool/UnrealBuildTool.csproj" -a ! -f "$BASE_PATH/../../../Build/InstalledBuild.txt" ]; then
-  dotnet msbuild /restore /target:build /property:Configuration=Development /nologo "$BASE_PATH/../../../Source/Programs/UnrealBuildTool/UnrealBuildTool.csproj" /verbosity:quiet
+	dotnet msbuild /restore /target:build /property:Configuration=Development /nologo "$BASE_PATH/../../../Source/Programs/UnrealBuildTool/UnrealBuildTool.csproj" /verbosity:quiet
 
-  if [ $? -ne 0 ]; then
-    echo GenerateProjectFiles ERROR: Failed to build UnrealBuildTool
-    exit 1
-  fi
+	if [ $? -ne 0 ]; then
+		echo GenerateProjectFiles ERROR: Failed to build UnrealBuildTool
+		exit 1
+	fi
 fi
 
 # pass all parameters to UBT
