Fixed a typo in an error message which listed an incorrect path.

diff --git a/Engine/Build/BatchFiles/Linux/GenerateProjectFiles.sh b/Engine/Build/BatchFiles/Linux/GenerateProjectFiles.sh
--- a/Engine/Build/BatchFiles/Linux/GenerateProjectFiles.sh
+++ b/Engine/Build/BatchFiles/Linux/GenerateProjectFiles.sh
@@ -16,7 +16,7 @@ if [ ! -d "$BASE_PATH/../../../Binaries/DotNET" ]; then
 fi
 
 if [ ! -d "$BASE_PATH/../../../Source" ]; then
-	echo GenerateProjectFiles ERROR: This script file does not appear to be located inside the Engine/Build/BatchFiles/Mac directory.
+	echo GenerateProjectFiles ERROR: This script file does not appear to be located inside the Engine/Build/BatchFiles/Linux directory.
 	exit 1
 fi
 
