Do not try to build UBT when running an installed builds (PR #8995 Thanks slonopotamus!)

diff --git a/Engine/Build/BatchFiles/Linux/GenerateProjectFiles.sh b/Engine/Build/BatchFiles/Linux/GenerateProjectFiles.sh
--- a/Engine/Build/BatchFiles/Linux/GenerateProjectFiles.sh
+++ b/Engine/Build/BatchFiles/Linux/GenerateProjectFiles.sh
@@ -22,11 +22,14 @@ fi
 
 source "$BASE_PATH/SetupEnvironment.sh" -dotnet "$BASE_PATH"
 
-if [ -f "$BASE_PATH/../../../Source/Programs/UnrealBuildTool/UnrealBuildTool.csproj" ]; then
-	dotnet msbuild /restore /target:build /property:Configuration=Development /nologo $BASE_PATH/../../../Source/Programs/UnrealBuildTool/UnrealBuildTool.csproj /verbosity:quiet
-	if [ $? -ne 0 ]; then
-		echo GenerateProjectFiles ERROR: Failed to build UnrealBuildTool
-		exit 1
+# If this is a source drop of the engine make sure that the UnrealBuildTool is up-to-date
+if [ ! -f $BASE_PATH/../../InstalledBuild.txt ]; then
+	if [ -f "$BASE_PATH/../../../Source/Programs/UnrealBuildTool/UnrealBuildTool.csproj" ]; then
+		dotnet msbuild /restore /target:build /property:Configuration=Development /nologo $BASE_PATH/../../../Source/Programs/UnrealBuildTool/UnrealBuildTool.csproj /verbosity:quiet
+		if [ $? -ne 0 ]; then
+			echo GenerateProjectFiles ERROR: Failed to build UnrealBuildTool
+			exit 1
+		fi
 	fi
 fi
 
