Update some files for building v20 13.0.1 toolchain

diff --git a/Engine/Build/BatchFiles/Linux/Toolchain/build_linux_toolchain.sh b/Engine/Build/BatchFiles/Linux/Toolchain/build_linux_toolchain.sh
--- a/Engine/Build/BatchFiles/Linux/Toolchain/build_linux_toolchain.sh
+++ b/Engine/Build/BatchFiles/Linux/Toolchain/build_linux_toolchain.sh
@@ -3,9 +3,10 @@
 set -x
 set -eu
 
-ToolChainVersion=v19
-LLVM_VERSION=11.0.1
+ToolChainVersion=v20
+LLVM_VERSION=13.0.1
 LLVM_URL=[REDACTED_PROFILE_URL]/llvm-project/releases/download/llvmorg-${LLVM_VERSION}
+ZLIB_PATH=/src/v1.2.8
 
 ToolChainVersionName="${ToolChainVersion}_clang-${LLVM_VERSION}-centos7"
 
@@ -22,8 +23,9 @@ umask 0022
 CORES=$(getconf _NPROCESSORS_ONLN)
 echo Using $CORES cores for building
 
+echo "check_certificate=off" > "$HOME/.wgetrc"
 # Get crosstool-ng
-git clone [REDACTED_PROFILE_URL]/crosstool-ng -b 1.22
+git clone [REDACTED_PROFILE_URL]/crosstool-ng -b 1.22
 
 # Build crosstool-ng
 pushd crosstool-ng
@@ -65,11 +67,13 @@ unset LD_LIBRARY_PATH
 # Build clang
 LLVM=llvm-${LLVM_VERSION}
 CLANG=clang-${LLVM_VERSION}
+UNWIND=libunwind-${LLVM_VERSION}
 LLD=lld-${LLVM_VERSION}
 COMPILER_RT=compiler-rt-${LLVM_VERSION}
 
 wget ${LLVM_URL}/${LLVM}.src.tar.xz
 wget ${LLVM_URL}/${CLANG}.src.tar.xz
+wget ${LLVM_URL}/${UNWIND}.src.tar.xz
 wget ${LLVM_URL}/${LLD}.src.tar.xz
 wget ${LLVM_URL}/${COMPILER_RT}.src.tar.xz
 
@@ -77,11 +81,21 @@ mkdir -p llvm
 tar -xf $LLVM.src.tar.xz --strip-components 1 -C llvm
 mkdir -p llvm/tools/clang
 tar -xf $CLANG.src.tar.xz --strip-components 1 -C llvm/tools/clang/
+mkdir -p llvm/libunwind
+tar -xf $UNWIND.src.tar.xz --strip-components 1 -C llvm/libunwind/
 mkdir -p llvm/tools/lld
 tar -xf $LLD.src.tar.xz --strip-components 1 -C llvm/tools/lld/
 mkdir -p llvm/projects/compiler-rt
 tar -xf $COMPILER_RT.src.tar.xz --strip-components 1 -C llvm/projects/compiler-rt/
 
+# this fixes an issue where AT_HWCAP2 is just not defined correctly in our sysroot. This is likely due to
+# AT_HWCAP2 being around since glibc 2.18 offically, while we are still stuck on 2.17 glibc.
+patch -d llvm/projects/compiler-rt/ -p 1 < /src/patches/compiler-rt/manually-define-AT_HWCAP2.diff
+
+# LLVM has just failed to support stand-alone LLD build, cheat by moving a required header into a location it can be found easily
+# if you fulling include this you end up breaking other things. [REDACTED_PROFILE_URL]/llvm-project/issues/48572
+cp -rf llvm/libunwind/include/mach-o/ llvm/include
+
 mkdir build-clang
 pushd build-clang
 	# CMake Error at cmake/modules/CheckCompilerVersion.cmake:40 (message):
@@ -94,11 +108,14 @@ pushd build-clang
 		-DCMAKE_BUILD_TYPE=Release \
 		-DLLVM_ENABLE_TERMINFO=OFF \
 		-DLLVM_ENABLE_LIBXML2=OFF \
+		-DLLVM_ENABLE_ZLIB=FORCE_ON \
+		-DZLIB_LIBRARY="$ZLIB_PATH/lib/Unix/x86_64-unknown-linux-gnu/libz_fPIC.a" \
+		-DZLIB_INCLUDE_DIR="$ZLIB_PATH/include/Unix/x86_64-unknown-linux-gnu" \
 		-DLLVM_ENABLE_LIBCXX=1 \
 		-DLLVM_TEMPORARILY_ALLOW_OLD_TOOLCHAIN=ON \
 		-DCMAKE_INSTALL_PREFIX=${InstallClangDir} \
 		-DLLVM_TARGETS_TO_BUILD="AArch64;X86"
-		
+
 	make -j$CORES && make install
 popd
 
@@ -167,13 +184,16 @@ for arch in $TARGETS; do
 			-DCMAKE_AR=${InstallClangDir}/bin/llvm-ar \
 			-DCMAKE_NM=${InstallClangDir}/bin/llvm-nm \
 			-DCMAKE_RANLIB=${InstallClangDir}/bin/llvm-ranlib \
-			-DLLVM_ENABLE_ZLIB=FORCE_ON
+			-DLLVM_ENABLE_ZLIB=FORCE_ON \
+			-DZLIB_LIBRARY="$ZLIB_PATH/lib/Unix/x86_64-unknown-linux-gnu/libz_fPIC.a" \
+			-DZLIB_INCLUDE_DIR="$ZLIB_PATH/include/Unix/x86_64-unknown-linux-gnu" \
 			-DCMAKE_EXE_LINKER_FLAGS="--target=$arch -L${OutputDirLinux}/$arch/lib64 --sysroot=${OutputDirLinux}/$arch -fuse-ld=lld" \
 			-DCMAKE_C_FLAGS="--target=$arch --sysroot=${OutputDirLinux}/$arch" \
 			-DCMAKE_CXX_FLAGS="--target=$arch --sysroot=${OutputDirLinux}/$arch" \
 			-DCMAKE_ASM_FLAGS="--target=$arch --sysroot=${OutputDirLinux}/$arch" \
 			-DCMAKE_INSTALL_PREFIX=../install-rt-$arch \
 			-DSANITIZER_COMMON_LINK_FLAGS="-fuse-ld=lld" \
+			-DSCUDO_LINK_FLAGS="-fuse-ld=lld" \
 			-DLLVM_CONFIG_PATH=${InstallClangDir}/bin/llvm-config
 
 		make -j$CORES && make install
