Forgot to submit change here for zlib when building toolchain

diff --git a/Engine/Build/BatchFiles/Linux/Toolchain/RunMe.bat b/Engine/Build/BatchFiles/Linux/Toolchain/RunMe.bat
--- a/Engine/Build/BatchFiles/Linux/Toolchain/RunMe.bat
+++ b/Engine/Build/BatchFiles/Linux/Toolchain/RunMe.bat
@@ -1,7 +1,7 @@
 @REM @echo off
 
-set TOOLCHAIN_VERSION=v19
-set LLVM_VERSION=11.0.1
+set TOOLCHAIN_VERSION=v20
+set LLVM_VERSION=13.0.1
 
 set SVN_BINARY=%CD%\..\..\..\..\Binaries\ThirdParty\svn\Win64\svn.exe
 set CMAKE_BINARY=%CD%\..\..\..\..\Extras\ThirdPartyNotUE\CMake\bin\cmake.exe
diff --git a/Engine/Build/BatchFiles/Linux/Toolchain/RunMe.sh b/Engine/Build/BatchFiles/Linux/Toolchain/RunMe.sh
--- a/Engine/Build/BatchFiles/Linux/Toolchain/RunMe.sh
+++ b/Engine/Build/BatchFiles/Linux/Toolchain/RunMe.sh
@@ -12,6 +12,11 @@ if ! $(grep -q "/docker/" /proc/1/cgroup); then
   # host commands
   ##############################################################################
 
+  ZLIB_PATH=../../../../../Engine/Source/ThirdParty/zlib/v1.2.8/
+
+  # Need to static link zlib for being able to compress debug files
+  cp -rpvf $ZLIB_PATH ./
+
   ImageName=build_linux_toolchain
 
   echo docker run -t --name ${ImageName} -v "${SCRIPT_DIR}:/src" centos:7 /src/${SCRIPT_NAME}
