- RunuAT now unsets some CVars that VisualStudio sets, which was causing UAT, when run from the editor that was run from VS, to have compile errors
#jria none

diff --git a/Engine/Build/BatchFiles/RunUAT.bat b/Engine/Build/BatchFiles/RunUAT.bat
--- a/Engine/Build/BatchFiles/RunUAT.bat
+++ b/Engine/Build/BatchFiles/RunUAT.bat
@@ -17,6 +17,17 @@ rem ## Change the CWD to /Engine.
 pushd "%~dp0..\..\"
 if not exist Build\BatchFiles\RunUAT.bat goto Error_BatchFileInWrongLocation
 
+
+rem Unset some env vars that are set when running from VisualStudio that can cause it to look outside of the bundled .net
+set VisualStudioDir=
+set VSSKUEDITION=
+rem Unset some others that don't cause errors, but could cause subtle differences between in-VS runs vs out-of-VS runs
+set PkgDefApplicationConfigFile
+set VSAPPIDDIR=
+set VisualStudioEdition=
+set VisualStudioVersion=
+
+
 set MSBUILD_LOGLEVEL=quiet
 set FORCECOMPILE_UAT=
 set NOCOMPILE_UAT=0
