RunUAT.bat: Force-uppercase the drive letter, to ensure consistency between 'dotnet msbuild AutomationTool.csproj' and AutomationTool's compilation of script modules.

#rnx

diff --git a/Engine/Build/BatchFiles/BuildUAT.bat b/Engine/Build/BatchFiles/BuildUAT.bat
--- a/Engine/Build/BatchFiles/BuildUAT.bat
+++ b/Engine/Build/BatchFiles/BuildUAT.bat
@@ -58,8 +58,6 @@ call "%~dp0GetDotnetPath.bat"
 if errorlevel 1 goto Error_NoDotnetSDK
 
 echo Building AutomationTool...
-rem ## print the timestamp on EpicGames.Core.dll in the hope of identifying if/where it is being modified.
-forfiles /P Programs\Shared\EpicGames.Core\obj\Development\netcoreapp3.1\ /M *.dll /c "cmd /c echo @path @fdate @ftime" 2> nul
 dotnet msbuild /restore /target:build /property:Configuration=Development /nologo Programs\AutomationTool\AutomationTool.csproj /verbosity:%MSBUILD_LOGLEVEL%
 if errorlevel 1 goto Error_UATCompileFailed
 
diff --git a/Engine/Build/BatchFiles/RunUAT.bat b/Engine/Build/BatchFiles/RunUAT.bat
--- a/Engine/Build/BatchFiles/RunUAT.bat
+++ b/Engine/Build/BatchFiles/RunUAT.bat
@@ -9,12 +9,18 @@ rem ## if you copy it to a different location and run it.
 setlocal EnableExtensions
 echo Running AutomationTool...
 
-set SCRIPT_DIR=%~dp0
+rem uppercase the drive letter
+SETLOCAL ENABLEDELAYEDEXPANSION
+set DRIVE_LETTER=%~d0
+FOR %%Z IN (A B C D E F G H I J K L M N O P Q R S T U V W X Y Z) DO SET DRIVE_LETTER=!DRIVE_LETTER:%%Z=%%Z!
+SETLOCAL DISABLEDELAYEDEXPANSION
+
+set SCRIPT_DIR=%DRIVE_LETTER%%~p0
 set UATExecutable=AutomationTool.dll
 set UATDirectory=Binaries\DotNET\AutomationTool
 
 rem ## Change the CWD to /Engine. 
-pushd "%~dp0..\..\"
+pushd "%SCRIPT_DIR%..\..\"
 if not exist Build\BatchFiles\RunUAT.bat goto Error_BatchFileInWrongLocation
 
 
