Fixed missing quotes in GenerateProjectFiles.sh

diff --git a/Engine/Build/BatchFiles/Linux/GenerateProjectFiles.sh b/Engine/Build/BatchFiles/Linux/GenerateProjectFiles.sh
--- a/Engine/Build/BatchFiles/Linux/GenerateProjectFiles.sh
+++ b/Engine/Build/BatchFiles/Linux/GenerateProjectFiles.sh
@@ -23,9 +23,9 @@ fi
 source "$BASE_PATH/SetupEnvironment.sh" -dotnet "$BASE_PATH"
 
 # If this is a source drop of the engine make sure that the UnrealBuildTool is up-to-date
-if [ ! -f $BASE_PATH/../../InstalledBuild.txt ]; then
+if [ ! -f "$BASE_PATH/../../InstalledBuild.txt" ]; then
 	if [ -f "$BASE_PATH/../../../Source/Programs/UnrealBuildTool/UnrealBuildTool.csproj" ]; then
-		dotnet msbuild /restore /target:build /property:Configuration=Development /nologo $BASE_PATH/../../../Source/Programs/UnrealBuildTool/UnrealBuildTool.csproj /verbosity:quiet
+		dotnet msbuild /restore /target:build /property:Configuration=Development /nologo "$BASE_PATH/../../../Source/Programs/UnrealBuildTool/UnrealBuildTool.csproj" /verbosity:quiet
 		if [ $? -ne 0 ]; then
 			echo GenerateProjectFiles ERROR: Failed to build UnrealBuildTool
 			exit 1
