Fix spam running RunUAT.bat ("Environment variable PkgDefApplicationConfigFile not defined"), caused by command not correctly clearing environment variable.

diff --git a/Engine/Build/BatchFiles/RunUAT.bat b/Engine/Build/BatchFiles/RunUAT.bat
--- a/Engine/Build/BatchFiles/RunUAT.bat
+++ b/Engine/Build/BatchFiles/RunUAT.bat
@@ -28,7 +28,7 @@ rem Unset some env vars that are set when running from VisualStudio that can cau
 set VisualStudioDir=
 set VSSKUEDITION=
 rem Unset some others that don't cause errors, but could cause subtle differences between in-VS runs vs out-of-VS runs
-set PkgDefApplicationConfigFile
+set PkgDefApplicationConfigFile=
 set VSAPPIDDIR=
 set VisualStudioEdition=
 set VisualStudioVersion=
