Remove more spam from RunUAT.bat caused by Turnkey hooks not correctly handling paths with spaces.

diff --git a/Engine/Build/BatchFiles/RunUAT.bat b/Engine/Build/BatchFiles/RunUAT.bat
--- a/Engine/Build/BatchFiles/RunUAT.bat
+++ b/Engine/Build/BatchFiles/RunUAT.bat
@@ -95,12 +95,12 @@ popd
 if %SET_TURNKEY_VARIABLES% == 0 goto SkipTurnkey
 
 rem ## Turnkey needs to update env vars in the calling process so that if it is run multiple times the Sdk env var changes are in effect
-if EXIST %SCRIPT_DIR%..\..\Intermediate\Turnkey\PostTurnkeyVariables.bat (
+if EXIST "%SCRIPT_DIR%..\..\Intermediate\Turnkey\PostTurnkeyVariables.bat" (
 	rem ## We need to endlocal so that the vars in the batch file work. NOTE: Working directory from pushd will be UNDONE here, but since we are about to quit, it's okay
 	endlocal 
 	echo Updating environment variables set by a Turnkey sub-process
-	call %SCRIPT_DIR%..\..\Intermediate\Turnkey\PostTurnkeyVariables.bat
-	del %SCRIPT_DIR%..\..\Intermediate\Turnkey\PostTurnkeyVariables.bat
+	call "%SCRIPT_DIR%..\..\Intermediate\Turnkey\PostTurnkeyVariables.bat"
+	del "%SCRIPT_DIR%..\..\Intermediate\Turnkey\PostTurnkeyVariables.bat"
 	rem ## setlocal again so that any popd's etc don't have an effect on calling process
 	setlocal
 )
