Change our default behviour for Build.sh.
Now we no longer build UBT by default, aligning with Windows Build.bat. You can still pass in: -buildubt as well as -buildscw to enable some extra build options

diff --git a/Engine/Build/BatchFiles/Linux/Build.sh b/Engine/Build/BatchFiles/Linux/Build.sh
--- a/Engine/Build/BatchFiles/Linux/Build.sh
+++ b/Engine/Build/BatchFiles/Linux/Build.sh
@@ -9,12 +9,28 @@ cd "`dirname "$0"`/../../../.."
 # Setup Environment and Mono
 source Engine/Build/BatchFiles/Linux/SetupEnvironment.sh -dotnet Engine/Build/BatchFiles/Linux
 
-# If this is a source drop of the engine make sure that the UnrealBuildTool is up-to-date
+# Skip UBT and SWC compile step if this is an installed build.
 if [ ! -f Engine/Build/InstalledBuild.txt ]; then
-	if ! dotnet build Engine/Source/Programs/UnrealBuildTool/UnrealBuildTool.csproj -c Development -v quiet; then
-		echo "Failed to build the build tool (UnrealBuildTool)"
-		exit 1
+	# use two for-loops here as -buildscw depends on UBT being built.
+	# So if -buildscw then -buildubt was passed it would fail/use out of date UBT
+	for i in "$@" ; do
+	# First make sure that the UnrealBuildTool is up-to-date
+	if [[ $i == "-buildubt" ]] ; then
+		if ! dotnet build Engine/Source/Programs/UnrealBuildTool/UnrealBuildTool.csproj -c Development -v quiet; then
+			echo "Failed to build the build tool (UnrealBuildTool)"
+			exit 1
+		fi
 	fi
+	done
+
+	for i in "$@" ; do
+	# build SCW if specified
+	if [[ $i == "-buildscw" ]] ; then
+		echo Building ShaderCompileWorker...
+		dotnet Engine/Binaries/DotNET/UnrealBuildTool/UnrealBuildTool.dll ShaderCompileWorker Linux Development
+		break
+	fi
+	done
 fi
 
 echo Running command : dotnet Engine/Binaries/DotNET/UnrealBuildTool/UnrealBuildTool.dll "$@"
diff --git a/Engine/Build/BatchFiles/Mac/Build.sh b/Engine/Build/BatchFiles/Mac/Build.sh
--- a/Engine/Build/BatchFiles/Mac/Build.sh
+++ b/Engine/Build/BatchFiles/Mac/Build.sh
@@ -9,14 +9,20 @@ source Engine/Build/BatchFiles/Mac/SetupEnvironment.sh -dotnet Engine/Build/Batc
 
 # Skip UBT and SWC compile step if this is an installed build.
 if [ ! -f Engine/Build/InstalledBuild.txt ]; then
+	# use two for-loops here as -buildscw depends on UBT being built.
+	# So if -buildscw then -buildubt was passed it would fail/use out of date UBT
+	for i in "$@" ; do
 	# First make sure that the UnrealBuildTool is up-to-date
-	if ! dotnet build Engine/Source/Programs/UnrealBuildTool/UnrealBuildTool.csproj -c Development -v quiet; then
-		echo "Failed to build the build tool (UnrealBuildTool)"
-		exit 1
+	if [[ $i == "-buildubt" ]] ; then
+		if ! dotnet build Engine/Source/Programs/UnrealBuildTool/UnrealBuildTool.csproj -c Development -v quiet; then
+			echo "Failed to build the build tool (UnrealBuildTool)"
+			exit 1
+		fi
 	fi
+	done
 
-	# build SCW if specified
 	for i in "$@" ; do
+	# build SCW if specified
 	if [[ $i == "-buildscw" ]] ; then
 		echo Building ShaderCompileWorker...
 		dotnet Engine/Binaries/DotNET/UnrealBuildTool/UnrealBuildTool.dll ShaderCompileWorker Mac Development
