Bundle libssl1/libcrpyto1 for dotnet as it is leaking to the system depends which will not be around for Ubuntu 22.04

https://discourse.ubuntu.com/t/openssl-3-0-transition-plans/24453

diff --git a/Engine/Build/BatchFiles/Linux/SetupDotnet.sh b/Engine/Build/BatchFiles/Linux/SetupDotnet.sh
--- a/Engine/Build/BatchFiles/Linux/SetupDotnet.sh
+++ b/Engine/Build/BatchFiles/Linux/SetupDotnet.sh
@@ -30,10 +30,15 @@ fi
 # Setup bundled Dotnet if cannot use installed one
 if [ $IS_DOTNET_INSTALLED -eq 0 ]; then
 	echo Setting up bundled DotNet SDK
-    CUR_DIR=`pwd`
-	export UE_DOTNET_DIR=$CUR_DIR/../../../Binaries/ThirdParty/DotNet/6.0.200/linux
-	export PATH=$UE_DOTNET_DIR:$PATH
-    export DOTNET_ROOT=$UE_DOTNET_DIR
+	CUR_DIR=`pwd`
+	export UE_DOTNET_DIR="$CUR_DIR/../../../Binaries/ThirdParty/DotNet/6.0.200/linux"
+	export PATH="$UE_DOTNET_DIR:$PATH"
+	export DOTNET_ROOT="$UE_DOTNET_DIR"
+
+	# We need to make sure point to our bundled libssl1, as ubuntu 22.04 is dropping libssl1 from the universe
+	# as well as force override for DotNet 6 to use 1.1 over 3 as we dont have that bundled atm
+	export CLR_OPENSSL_VERSION_OVERRIDE=1.1
+	export LD_LIBRARY_PATH="$CUR_DIR/../../../Binaries/ThirdParty/OpenSSL/Unix/lib/x86_64-unknown-linux-gnu:$LD_LIBRARY_PATH"
 else
 	export IS_DOTNET_INSTALLED=$IS_DOTNET_INSTALLED
 fi
