Fix spelling

#rnx
#prelfight none

diff --git a/Engine/Build/BatchFiles/Linux/Build.sh b/Engine/Build/BatchFiles/Linux/Build.sh
--- a/Engine/Build/BatchFiles/Linux/Build.sh
+++ b/Engine/Build/BatchFiles/Linux/Build.sh
@@ -9,7 +9,7 @@ cd "`dirname "$0"`/../../../.."
 # Setup Environment and Mono
 source Engine/Build/BatchFiles/Linux/SetupEnvironment.sh -dotnet Engine/Build/BatchFiles/Linux
 
-# Skip UBT and SWC compile step if this is an installed build.
+# Skip UBT and SCW compile step if this is an installed build.
 if [ ! -f Engine/Build/InstalledBuild.txt ]; then
 	# use two for-loops here as -buildscw depends on UBT being built.
 	# So if -buildscw then -buildubt was passed it would fail/use out of date UBT
diff --git a/Engine/Build/BatchFiles/Mac/Build.sh b/Engine/Build/BatchFiles/Mac/Build.sh
--- a/Engine/Build/BatchFiles/Mac/Build.sh
+++ b/Engine/Build/BatchFiles/Mac/Build.sh
@@ -7,7 +7,7 @@ cd "`dirname "$0"`/../../../.."
 # Setup Environment for DotNET
 source Engine/Build/BatchFiles/Mac/SetupEnvironment.sh -dotnet Engine/Build/BatchFiles/Mac
 
-# Skip UBT and SWC compile step if this is an installed build.
+# Skip UBT and SCW compile step if this is an installed build.
 if [ ! -f Engine/Build/InstalledBuild.txt ]; then
 	# use two for-loops here as -buildscw depends on UBT being built.
 	# So if -buildscw then -buildubt was passed it would fail/use out of date UBT
