- Changed RunUAT.sh to always chmod u+x dotnet

diff --git a/Engine/Build/BatchFiles/RunUAT.sh b/Engine/Build/BatchFiles/RunUAT.sh
--- a/Engine/Build/BatchFiles/RunUAT.sh
+++ b/Engine/Build/BatchFiles/RunUAT.sh
@@ -119,6 +119,9 @@ else
 	LogDir="$uebp_LogFolder"
 fi
 
+# make sure it's executable (if it was zipped it may not be)
+chmod u+x `which dotnet`
+
 # if we are running under UE, we need to run this with the term handler (otherwise canceling a UAT job from the editor
 # can leave mono, etc running in the background, which means we need the PID so we 
 # run it in the background
