Update SetupToolchain.sh to download new v20 toolchain

diff --git a/Engine/Build/BatchFiles/Linux/SetupToolchain.sh b/Engine/Build/BatchFiles/Linux/SetupToolchain.sh
--- a/Engine/Build/BatchFiles/Linux/SetupToolchain.sh
+++ b/Engine/Build/BatchFiles/Linux/SetupToolchain.sh
@@ -4,7 +4,7 @@
 # put ourselves into Engine directory (two up from location of this script)
 pushd "`dirname "$0"`/../../.." > /dev/null
 
-TOOLCHAIN_VERSION=v19_clang-11.0.1-centos7
+TOOLCHAIN_VERSION=v20_clang-13.0.1-centos7
 TOOLCHAIN_ARCHIVE=$TOOLCHAIN_VERSION.tar.gz
 
 TOOLCHAIN_URL=http://cdn.unrealengine.com/Toolchain_Linux/native-linux-$TOOLCHAIN_ARCHIVE
