- MIssed NFL path in P4VUtils installers

diff --git a/Engine/Build/BatchFiles/InstallP4VUtils.bat b/Engine/Build/BatchFiles/InstallP4VUtils.bat
--- a/Engine/Build/BatchFiles/InstallP4VUtils.bat
+++ b/Engine/Build/BatchFiles/InstallP4VUtils.bat
@@ -3,11 +3,11 @@ setlocal
 set LIBPATH=%LOCALAPPDATA%\Unreal Engine\P4VUtils
 
 echo Copying P4VUtils files to %LIBPATH%...
-mkdir "%LIBPATH%"
+mkdir "%LIBPATH%" 2> NUL
 
-copy ../../../Extras/P4VUtils/Binaries/* "%LIBPATH%"
-copy ../../../Extras/P4VUtils/P4VUtils.ini "%LIBPATH%"
-copy ../../../Restricted/NotForLicensees/* "%LIBPATH%" 2> NUL
+copy ..\..\Extras\P4VUtils\Binaries\Win64\* "%LIBPATH%"
+copy ..\..\Extras\P4VUtils\P4VUtils.ini "%LIBPATH%"
+xcopy /s /y ..\..\Restricted\NotForLicensees\Extras\P4VUtils\* "%LIBPATH%\NotForLicensees\" 2> NUL
 
 echo Installing P4VUtils into p4v...
 dotnet "%LIBPATH%\P4VUtils.dll" install
diff --git a/Engine/Build/BatchFiles/Linux/InstallP4VUtils.sh b/Engine/Build/BatchFiles/Linux/InstallP4VUtils.sh
--- a/Engine/Build/BatchFiles/Linux/InstallP4VUtils.sh
+++ b/Engine/Build/BatchFiles/Linux/InstallP4VUtils.sh
@@ -8,9 +8,12 @@ LibPath="$HOME/Library/Unreal Engine/P4VUtils"
 echo
 echo Copying P4VUtils files to '$LibPath'...
 mkdir -p "$LibPath"
-cp ../../../Extras/P4VUtils/Binaries/Linux/* "$LibPath"
-cp ../../../Extras/P4VUtils/P4VUtils.ini "$LibPath"
-cp ../../../Restricted/NotForLicensees/* "$LibPath" 2> /dev/null
+cp -f ../../../Extras/P4VUtils/Binaries/Linux/* "$LibPath"
+cp -f ../../../Extras/P4VUtils/P4VUtils.ini "$LibPath"
+if [ -e ../../../Restricted/NotForLicensees/Extras/P4VUtils ]; then
+	mkdir "$LibPath/NotForLicensees"
+	cp -f ../../../Restricted/NotForLicensees/Extras/P4VUtils/* "$LibPath/NotForLicensees/"
+fi
 
 echo
 echo Installing P4VUtils into p4v...
diff --git a/Engine/Build/BatchFiles/Mac/InstallP4VUtils.command b/Engine/Build/BatchFiles/Mac/InstallP4VUtils.command
--- a/Engine/Build/BatchFiles/Mac/InstallP4VUtils.command
+++ b/Engine/Build/BatchFiles/Mac/InstallP4VUtils.command
@@ -8,9 +8,12 @@ LibPath="$HOME/Library/Unreal Engine/P4VUtils"
 echo
 echo Copying P4VUtils files to '$LibPath'...
 mkdir -p "$LibPath"
-cp ../../../Extras/P4VUtils/Binaries/Mac/* "$LibPath"
-cp ../../../Extras/P4VUtils/P4VUtils.ini "$LibPath"
-cp ../../../Restricted/NotForLicensees/* "$LibPath" 2> /dev/null
+cp -f ../../../Extras/P4VUtils/Binaries/Mac/* "$LibPath"
+cp -f ../../../Extras/P4VUtils/P4VUtils.ini "$LibPath"
+if [ -e ../../../Restricted/NotForLicensees/Extras/P4VUtils ]; then
+	mkdir "$LibPath/NotForLicensees"
+	cp -f ../../../Restricted/NotForLicensees/Extras/P4VUtils/* "$LibPath/NotForLicensees/"
+fi
 
 echo
 echo Installing P4VUtils into p4v...
