- Added a CompileSingleFile.sh that can run in Xcode to compile a single file via UBT. It activates Terminal so the user can see output and copy and paste, etc. To use, make a new custom Behavior (+ button at bottom of Behaviors window), and point it at this script. It uses AppleScript to find the current document, platform, and scheme, but will always compile in Development (couldn't find a way to get the project to get the build config)

diff --git a/Engine/Build/BatchFiles/Mac/CompileSingleFile.sh b/Engine/Build/BatchFiles/Mac/CompileSingleFile.sh
new file mode 100755
--- /dev/null
+++ b/Engine/Build/BatchFiles/Mac/CompileSingleFile.sh
@@ -0,0 +1,39 @@
+#!/bin/bash
+
+# go to batch files dirs
+cd $(dirname "$0")
+
+export DOC=$(osascript -e 'tell application "Xcode" 
+	return path of document 1 whose name ends with (word -1 of (get name of window 1))
+end tell')
+export TARGET=$(osascript -e 'tell application "Xcode" 
+	return name of active scheme of workspace document 1
+end tell')
+PLATFORM=$(osascript -e 'tell application "Xcode" 
+	return platform of active run destination of workspace document 1
+end tell')
+
+PLATFORM=${PLATFORM/macosx/Mac}
+PLATFORM=${PLATFORM/iphoneos/IOS}
+export PLATFORM=${PLATFORM/tvos/TVOS}
+
+echo Doc: ${DOC}
+echo Target: ${TARGET}
+echo Platform: ${PLATFORM}
+
+
+osascript -e 'tell application "Terminal"
+	activate
+	set Dir to (system attribute "PWD")
+	set Target to (system attribute "TARGET")
+	set Platform to (system attribute "PLATFORM")
+	set Doc to (system attribute "DOC")
+	do script Dir & "/Build.sh " & Target & " Development " & Platform & " -singlefile=\"" & Doc & "\"; exit" 
+end tell'
+# ./Build.sh ${TARGET} Development ${PLATFORM} -singlefile="${DOC}"
+
+# echo ${XCODEINFO}
+
+# echo "-${ActivePath}-"
+
+# osascript -e 'tell application "Xcode" to display dialog (system attribute "FILEPATH") buttons {"OK"} default button 1'
