Accepting GitHub 9284 which modifies SetupToolchain.sh to not use incomplete archives when downloading the toolchain (instead it starts from scratch)

diff --git a/Engine/Build/BatchFiles/Linux/SetupToolchain.sh b/Engine/Build/BatchFiles/Linux/SetupToolchain.sh
--- a/Engine/Build/BatchFiles/Linux/SetupToolchain.sh
+++ b/Engine/Build/BatchFiles/Linux/SetupToolchain.sh
@@ -34,14 +34,16 @@ else
 fi
 
 if [ ! -f $TOOLCHAIN_ARCHIVE ]; then
+    rm -f ${TOOLCHAIN_ARCHIVE}.tmp
 	if which curl 1>/dev/null; then
-		curl $TOOLCHAIN_URL -o $TOOLCHAIN_ARCHIVE
+		curl $TOOLCHAIN_URL -o ${TOOLCHAIN_ARCHIVE}.tmp || exit 1
 	elif which wget 1>/dev/null; then
-		wget $TOOLCHAIN_URL -O $TOOLCHAIN_ARCHIVE
+		wget $TOOLCHAIN_URL -O ${TOOLCHAIN_ARCHIVE}.tmp || exit 1
 	else 
 		echo "Please install curl or wget"
 		exit
 	fi
+    mv ${TOOLCHAIN_ARCHIVE}.tmp $TOOLCHAIN_ARCHIVE
 else
 	echo "Using cached toolchain."
 fi
@@ -70,3 +72,4 @@ if [ -f $TOOLCHAIN_ARCHIVE ]; then
 else
 	echo "Download failed please fetch from $TOOLCHAIN_URL manually".
 fi
+
