Wrap paths in quotes to allow for spaces in path names

#rnx

diff --git a/Engine/Build/BatchFiles/Mac/GitDependencies.sh b/Engine/Build/BatchFiles/Mac/GitDependencies.sh
--- a/Engine/Build/BatchFiles/Mac/GitDependencies.sh
+++ b/Engine/Build/BatchFiles/Mac/GitDependencies.sh
@@ -5,7 +5,7 @@ set -e
 
 SCRIPT_PATH=$0
 if [ -L "$SCRIPT_PATH" ]; then
-	SCRIPT_PATH=$(dirname "$SCRIPT_PATH")/$(readlink "$SCRIPT_PATH")
+	SCRIPT_PATH="$(dirname "$SCRIPT_PATH")/$(readlink "$SCRIPT_PATH")"
 fi
 
 cd "$(dirname "$SCRIPT_PATH")" && SCRIPT_PATH="`pwd`/$(basename "$SCRIPT_PATH")"
@@ -15,7 +15,7 @@ if [ ! -f ../../../Binaries/DotNET/GitDependencies/osx-x64/GitDependencies ]; th
 	exit 1
 fi 
 
-$(dirname "$SCRIPT_PATH")/SetupDotnet.sh
+"$(dirname "$SCRIPT_PATH")/SetupDotnet.sh"
 
 cd ../../../..
 
