Preserve the ERRORLEVEL from AutomationTool if we apply Turnkey variables.

#rnx

diff --git a/Engine/Build/BatchFiles/RunUAT.bat b/Engine/Build/BatchFiles/RunUAT.bat
--- a/Engine/Build/BatchFiles/RunUAT.bat
+++ b/Engine/Build/BatchFiles/RunUAT.bat
@@ -91,13 +91,14 @@ rem ## Run AutomationTool
 pushd %UATDirectory%
 dotnet %UATExecutable% %*
 popd
+set RUNUAT_ERRORLEVEL=%ERRORLEVEL%
 
 if %SET_TURNKEY_VARIABLES% == 0 goto SkipTurnkey
 
 rem ## Turnkey needs to update env vars in the calling process so that if it is run multiple times the Sdk env var changes are in effect
 if EXIST "%SCRIPT_DIR%..\..\Intermediate\Turnkey\PostTurnkeyVariables.bat" (
-	rem ## We need to endlocal so that the vars in the batch file work. NOTE: Working directory from pushd will be UNDONE here, but since we are about to quit, it's okay
-	endlocal 
+	rem ## We need to endlocal so that the vars in the batch file work. NOTE: Working directory from pushd will be UNDONE here, but since we are about to quit, it's okay. UAT errorlevel is preserved beyond the endlocal
+	endlocal & set RUNUAT_ERRORLEVEL=%RUNUAT_ERRORLEVEL%
 	echo Updating environment variables set by a Turnkey sub-process
 	call "%SCRIPT_DIR%..\..\Intermediate\Turnkey\PostTurnkeyVariables.bat"
 	del "%SCRIPT_DIR%..\..\Intermediate\Turnkey\PostTurnkeyVariables.bat"
@@ -106,7 +107,7 @@ if EXIST "%SCRIPT_DIR%..\..\Intermediate\Turnkey\PostTurnkeyVariables.bat" (
 )
 :SkipTurnkey
 
-if not %ERRORLEVEL% == 0 goto Error_UATFailed
+if not %RUNUAT_ERRORLEVEL% == 0 goto Error_UATFailed
 
 rem ## Success!
 goto Exit
@@ -134,7 +135,7 @@ goto Exit_Failure
 
 
 :Error_UATFailed
-set RUNUAT_EXITCODE=%ERRORLEVEL%
+set RUNUAT_EXITCODE=%RUNUAT_ERRORLEVEL%
 goto Exit_Failure
 
 :Exit_Failure
