Disabled an optimization which crashes dotnet during GenerateProjectFiles on Ubuntu 20.04.

diff --git a/Engine/Build/BatchFiles/Linux/SetupDotnet.sh b/Engine/Build/BatchFiles/Linux/SetupDotnet.sh
--- a/Engine/Build/BatchFiles/Linux/SetupDotnet.sh
+++ b/Engine/Build/BatchFiles/Linux/SetupDotnet.sh
@@ -31,6 +31,10 @@ fi
 if [ $IS_DOTNET_INSTALLED -eq 0 ]; then
 	echo Setting up bundled DotNet SDK
 	CUR_DIR=`pwd`
+
+	# If this flag isn't set to 0, dotnet crashes during GenerateProjectFiles.sh on Ubuntu 20.04 
+	export DOTNET_gcServer=0
+
 	export UE_DOTNET_DIR="$CUR_DIR/../../../Binaries/ThirdParty/DotNet/6.0.302/linux"
 	chmod u+x "$UE_DOTNET_DIR/dotnet"
 	export PATH="$UE_DOTNET_DIR:$PATH"
