Disable .NET Server GC for Mac as it is currently causing errors

#rnx

diff --git a/Engine/Build/BatchFiles/Mac/SetupDotnet.sh b/Engine/Build/BatchFiles/Mac/SetupDotnet.sh
--- a/Engine/Build/BatchFiles/Mac/SetupDotnet.sh
+++ b/Engine/Build/BatchFiles/Mac/SetupDotnet.sh
@@ -7,6 +7,9 @@ cd "$1"
 IS_DOTNET_INSTALLED=0
 DOTNET_VERSION_PATH=$(command -v dotnet) || true
 
+# gcServer can cause "An error occurred trying to start process 'X' with working directory 'Y'. Too many open files"
+export DOTNET_gcServer=0
+
 if [ "$UE_USE_SYSTEM_DOTNET" == "1" ] && [ ! $DOTNET_VERSION_PATH == "" ] && [ -f $DOTNET_VERSION_PATH ]; then
 	# If dotnet is installed, check that it has a new enough version of the SDK
 	DOTNET_SDKS=(`dotnet --list-sdks | grep -P "(\d*)\.(\d*)\..* \[(.*)\]"`)
