Fix ShaderCompileWorker arm64 support in Xcode

diff --git a/Engine/Build/BatchFiles/Mac/XcodeBuild.sh b/Engine/Build/BatchFiles/Mac/XcodeBuild.sh
--- a/Engine/Build/BatchFiles/Mac/XcodeBuild.sh
+++ b/Engine/Build/BatchFiles/Mac/XcodeBuild.sh
@@ -108,7 +108,7 @@ if [ "$ACTION" == "build" ]; then
 
 	# Build SCW if this is an editor target
 	if [[ "$TARGET" == *"Editor" ]]; then
-		dotnet Engine/Binaries/DotNET/UnrealBuildTool/UnrealBuildTool.dll ShaderCompileWorker Mac Development
+		dotnet Engine/Binaries/DotNET/UnrealBuildTool/UnrealBuildTool.dll ShaderCompileWorker Mac Development $UBT_ARCHFLAG
 	fi
 
 elif [ $ACTION == "clean" ]; then
