Fix up code compile batch scripts to better support visual studio
* If invoked from VS2022, use visual studio to build UnrealBuildTool to prevent unnecessary rebuilds
* If invoked from VS2019 or on the command line without visual studio, use bundled dotnet to build
* Deduplicate all logic into Build.bat, and have Clean.bat and Rebuild.bat call Build.bat with either -Clean or -Rebuild
* Print entire dotnet version
* Replace various calls to dotnet msbuild with dotnet build
* Update old UE4 comments

diff --git a/Engine/Build/BatchFiles/Build.bat b/Engine/Build/BatchFiles/Build.bat
--- a/Engine/Build/BatchFiles/Build.bat
+++ b/Engine/Build/BatchFiles/Build.bat
@@ -1,32 +1,79 @@
 @echo off
+
+rem ## Unreal Engine code build script
+rem ## Copyright Epic Games, Inc. All Rights Reserved.
+rem ##
+rem ## This script is expecting to exist in the /Engine/Build/BatchFiles directory.  It will not work correctly
+rem ## if you copy it to a different location and run it.
+rem ##
+rem ##     %1 is the game name
+rem ##     %2 is the platform name
+rem ##     %3 is the configuration name
+rem ##     additional args are passed directly to UnrealBuildTool
+
 setlocal enabledelayedexpansion
 
-REM The %~dp0 specifier resolves to the path to the directory where this .bat is located in.
-REM We use this so that regardless of where the .bat file was executed from, we can change to
-REM directory relative to where we know the .bat is stored.
+rem ## First, make sure the batch file exists in the folder we expect it to.  This is necessary in order to
+rem ## verify that our relative path to the /Engine/Source directory is correct
+if not exist "%~dp0..\..\Source" goto Error_BatchFileInWrongLocation
+
+rem ## Change the CWD to /Engine/Source.  We always need to run UnrealBuildTool from /Engine/Source!
 pushd "%~dp0\..\..\Source"
+if not exist ..\Build\BatchFiles\Clean.bat goto Error_BatchFileInWrongLocation
+
+set UBTPath="..\..\Engine\Binaries\DotNET\UnrealBuildTool\UnrealBuildTool.dll"
+
+rem ## If this is an installed build, we don't need to rebuild UBT. Go straight to building.
+if exist ..\Build\InstalledBuild.txt goto ReadyToBuild
 
 rem ## Verify that dotnet is present
 call "%~dp0GetDotnetPath.bat"
+if errorlevel 1 goto Error_NoDotnetSDK
+REM ## Skip msbuild detection if using dotnet as this is done for us by dotnet-cli
 
-REM %1 is the game name
-REM %2 is the platform name
-REM %3 is the configuration name
-
-set UBTPath="..\..\Engine\Binaries\DotNET\UnrealBuildTool\UnrealBuildTool.dll"
+rem ## Compile UBT if the project file exists
+:ReadyToBuildUBT
+set ProjectFile="Programs\UnrealBuildTool\UnrealBuildTool.csproj"
+if not exist %ProjectFile% goto NoProjectFile
 
-IF EXIST %UBTPath% (
-		dotnet %UBTPath% %*
-		popd
-
-		REM Ignore exit codes of 2 ("ECompilationResult.UpToDate") from UBT; it's not a failure.
-		if "!ERRORLEVEL!"=="2" (
-			EXIT /B 0
-		)
-		 
-		EXIT /B !ERRORLEVEL!
-) ELSE (
-	ECHO UnrealBuildTool.dll not found in %UBTPath%
-	popd
-	EXIT /B 999
+rem ## If this script was called from Visual Studio 2022, build UBT with Visual Studio to prevent unnecessary rebuilds.
+if "%VisualStudioVersion%" GEQ "17.0" (
+	echo Building UnrealBuildTool with %VisualStudioEdition%...
+	"%VSAPPIDDIR%..\..\MSBuild\Current\Bin\MSBuild.exe" %ProjectFile% -t:Build -p:Configuration=Development -verbosity:quiet -noLogo
+	if errorlevel 1 goto Error_UBTCompileFailed
+) else (
+	echo Building UnrealBuildTool with dotnet...
+	dotnet build %ProjectFile% -c Development -v quiet
+	if errorlevel 1 goto Error_UBTCompileFailed
 )
+:NoProjectFile
+
+rem ## Run UBT
+:ReadyToBuild
+if not exist %UBTPath% goto Error_UBTMissing
+echo Running UnrealBuildTool: dotnet %UBTPath% %*
+dotnet %UBTPath% %*
+goto Exit
+
+:Error_BatchFileInWrongLocation
+echo ERROR: The batch file does not appear to be located in the Engine/Build/BatchFiles directory. This script must be run from within that directory.
+pause
+goto Exit
+
+:Error_NoDotnetSDK
+echo ERROR: Unable to find a install of Dotnet SDK.  Please make sure you have it installed and that `dotnet` is a globally available command.
+pause
+goto Exit
+
+:Error_UBTCompileFailed
+echo ERROR: Failed to build UnrealBuildTool.
+pause
+goto Exit
+
+:Error_UBTMissing
+echo ERROR: UnrealBuildTool.dll not found in %UBTPath%
+pause
+goto Exit
+
+:Exit
+
diff --git a/Engine/Build/BatchFiles/BuildUAT.bat b/Engine/Build/BatchFiles/BuildUAT.bat
--- a/Engine/Build/BatchFiles/BuildUAT.bat
+++ b/Engine/Build/BatchFiles/BuildUAT.bat
@@ -60,7 +60,7 @@ call "%~dp0GetDotnetPath.bat"
 if errorlevel 1 goto Error_NoDotnetSDK
 
 echo Building AutomationTool...
-dotnet msbuild /restore /target:build /property:Configuration=Development /nologo Programs\AutomationTool\AutomationTool.csproj /verbosity:%MSBUILD_LOGLEVEL%
+dotnet build Programs\AutomationTool\AutomationTool.csproj -c Development -v %MSBUILD_LOGLEVEL%
 if errorlevel 1 goto Error_UATCompileFailed
 
 rem record input files - regardless of how we got here, these are now our point of reference
diff --git a/Engine/Build/BatchFiles/BuildUAT.sh b/Engine/Build/BatchFiles/BuildUAT.sh
--- a/Engine/Build/BatchFiles/BuildUAT.sh
+++ b/Engine/Build/BatchFiles/BuildUAT.sh
@@ -77,8 +77,7 @@ if [ $PERFORM_REBUILD -eq 1 ]; then
     source "$SCRIPT_DIR/Linux/SetupEnvironment.sh" $EnvironmentType "$SCRIPT_DIR/Linux"
   fi
 
-  dotnet msbuild -restore -t:build -property:Configuration=Development -nologo \
-    Programs/AutomationTool/AutomationTool.csproj -verbosity:$MSBuild_Verbosity
+  dotnet build Programs/AutomationTool/AutomationTool.csproj -c Development -v $MSBuild_Verbosity
   if [ $? -ne 0 ]; then
     echo "Compilation failed"
     popd >/dev/null
diff --git a/Engine/Build/BatchFiles/Clean.bat b/Engine/Build/BatchFiles/Clean.bat
--- a/Engine/Build/BatchFiles/Clean.bat
+++ b/Engine/Build/BatchFiles/Clean.bat
@@ -1,76 +1,23 @@
 @echo off
 
-rem ## Unreal Engine 4 cleanup script
+rem ## Unreal Engine code cleanup script
 rem ## Copyright Epic Games, Inc. All Rights Reserved.
 rem ##
-rem ## This script is expecting to exist in the UE4 root directory.  It will not work correctly
+rem ## This script is expecting to exist in the Engine/Build/BatchFiles directory.  It will not work correctly
 rem ## if you copy it to a different location and run it.
 rem ##
 rem ##     %1 is the game name
 rem ##     %2 is the platform name
 rem ##     %3 is the configuration name
+rem ##     additional args are passed directly to UnrealBuildTool
 
-setlocal
+IF NOT EXIST "%~dp0\Build.bat" GOTO Error_MissingBuildBatchFile
 
-rem ## First, make sure the batch file exists in the folder we expect it to.  This is necessary in order to
-rem ## verify that our relative path to the /Engine/Source directory is correct
-if not exist "%~dp0..\..\Source" goto Error_BatchFileInWrongLocation
-
-
-rem ## Change the CWD to /Engine/Source.  We always need to run UnrealBuildTool from /Engine/Source!
-pushd "%~dp0..\..\Source"
-if not exist ..\Build\BatchFiles\Clean.bat goto Error_BatchFileInWrongLocation
-
-set UBTPath="..\..\Engine\Binaries\DotNET\UnrealBuildTool\UnrealBuildTool.dll"
-
-rem ## If this is an installed build, we don't need to rebuild UBT. Go straight to cleaning.
-if exist ..\Build\InstalledBuild.txt goto ReadyToClean
-
-rem ## Get the path to MSBuild
-call "%~dp0GetMSBuildPath.bat"
-if errorlevel 1 goto Error_NoVisualStudioEnvironment
-
-
-rem ## Compile UBT if the project file exists
-
-set ProjectFile="Programs\UnrealBuildTool\UnrealBuildTool.csproj"
-if not exist %ProjectFile% goto NoProjectFile
-dotnet build Programs\UnrealBuildTool\UnrealBuildTool.csproj -c Development -v quiet
-if errorlevel 1 goto Error_UBTCompileFailed
-:NoProjectFile
-
-
-rem ## Execute UBT
-:ReadyToClean
-if not exist %UBTPath% goto Error_UBTMissing
-dotnet %UBTPath% %* -clean
+call "%~dp0\Build.bat" %* -Clean
 goto Exit
 
-
-:Error_BatchFileInWrongLocation
-echo Clean ERROR: The batch file does not appear to be located in the UE4/Engine/Build/BatchFiles directory.  This script must be run from within that directory.
-pause
-goto Exit
-
-:Error_NoVisualStudioEnvironment
-echo GenerateProjectFiles ERROR: A valid version of Visual Studio does not appear to be installed.
-pause
-goto Exit
-
-:Error_UBTCompileFailed
-echo Clean ERROR: Failed to build UnrealBuildTool.
-pause
-goto Exit
-
-:Error_UBTMissing
-echo Clean ERROR: UnrealBuildTool.dll not found in %UBTPath%
-pause
-goto Exit
-
-:Error_CleanFailed
-echo Clean ERROR: Clean failed.
-pause
-goto Exit
+:Error_MissingBuildBatchFile
+ECHO Build.bat not found in "%~dp0"
+EXIT /B 999
 
 :Exit
-
diff --git a/Engine/Build/BatchFiles/GenerateProjectFiles.bat b/Engine/Build/BatchFiles/GenerateProjectFiles.bat
--- a/Engine/Build/BatchFiles/GenerateProjectFiles.bat
+++ b/Engine/Build/BatchFiles/GenerateProjectFiles.bat
@@ -1,13 +1,13 @@
 @echo off
 
-rem ## Unreal Engine 5 Visual Studio project setup script
+rem ## Unreal Engine Visual Studio project setup script
 rem ## Copyright Epic Games, Inc. All Rights Reserved.
 rem ##
 rem ## This script is expecting to exist in the UE5 root directory.  It will not work correctly
 rem ## if you copy it to a different location and run it.
 
 setlocal
-echo Setting up Unreal Engine 5 project files...
+echo Setting up Unreal Engine project files...
 
 rem ## First, make sure the batch file exists in the folder we expect it to.  This is necessary in order to
 rem ## verify that our relative path to the /Engine/Source directory is correct
@@ -48,11 +48,11 @@ fc /b ..\Intermediate\Build\UnrealBuildToolFiles.txt ..\Intermediate\Build\Unrea
 if not errorlevel 1 goto SkipClean
 copy /y ..\Intermediate\Build\UnrealBuildToolFiles.txt ..\Intermediate\Build\UnrealBuildToolPrevFiles.txt >nul
 
-dotnet msbuild /target:clean /property:Configuration=Development /nologo Programs\UnrealBuildTool\UnrealBuildTool.csproj /verbosity:quiet
+dotnet clean Programs\UnrealBuildTool\UnrealBuildTool.csproj -c Development -v quiet
 
 :SkipClean
 echo Building UnrealBuildTool...
-dotnet msbuild /restore /target:build /property:Configuration=Development /nologo Programs\UnrealBuildTool\UnrealBuildTool.csproj /verbosity:quiet
+dotnet build Programs\UnrealBuildTool\UnrealBuildTool.csproj -c Development -v quiet
 
 if errorlevel 1 goto Error_UBTCompileFailed
 
diff --git a/Engine/Build/BatchFiles/GetDotnetPath.bat b/Engine/Build/BatchFiles/GetDotnetPath.bat
--- a/Engine/Build/BatchFiles/GetDotnetPath.bat
+++ b/Engine/Build/BatchFiles/GetDotnetPath.bat
@@ -1,6 +1,6 @@
 @echo off
 
-rem ## Unreal Engine 5 utility script
+rem ## Unreal Engine utility script
 rem ## Copyright Epic Games, Inc. All Rights Reserved.
 rem ##
 rem ## This script verifies that dotnet sdk is installed and a new enough SDK is present, alternatively setups up the submitted SDK for use.
@@ -9,7 +9,8 @@ rem if UE_USE_SYSTEM_DOTNET we assume a installed dotnet is present
 if "%UE_USE_SYSTEM_DOTNET%" == "1" goto verify_dotnet
 
 rem add the dotnet sdk in the repo as the current dotnet sdk
-set UE_DOTNET_DIR=%~dp0..\..\Binaries\ThirdParty\DotNet\6.0.302\windows
+set UE_DOTNET_VERSION=6.0.302
+set UE_DOTNET_DIR=%~dp0..\..\Binaries\ThirdParty\DotNet\%UE_DOTNET_VERSION%\windows
 set PATH=%UE_DOTNET_DIR%;%PATH%
 set DOTNET_ROOT=%UE_DOTNET_DIR%
 
@@ -20,7 +21,7 @@ rem for dotnet applications that require older dotnet runtimes, this will allow
 set DOTNET_ROLL_FORWARD=LatestMajor
 
 rem skip dotnet verification when using our submitted sdk as we know it is up to date
-ECHO Using bundled DotNet SDK
+ECHO Using bundled DotNet SDK version: %UE_DOTNET_VERSION%
 goto Succeeded
 
 :verify_dotnet
@@ -53,7 +54,7 @@ for /f "tokens=1,* delims= " %%I in ('dotnet --list-sdks') do (
 
 				set FOUND_MAJOR=%%X
 				set FOUND_MINOR=%%Y
-				ECHO Found Dotnet SDK version: %%X.%%Y
+				ECHO Found Dotnet SDK version: %%X.%%Y.%%Z
 				goto Succeeded
 			)
 		)
@@ -62,7 +63,7 @@ for /f "tokens=1,* delims= " %%I in ('dotnet --list-sdks') do (
 			REM If the major version is greater then what we require then this sdk is good enough
 			set FOUND_MAJOR=%%X
 			set FOUND_MINOR=%%Y
-			ECHO Found Dotnet SDK version: %%X.%%Y
+			ECHO Found Dotnet SDK version: %%X.%%Y.%%Z
 			goto Succeeded
 		)
 
diff --git a/Engine/Build/BatchFiles/GetMSBuildPath.bat b/Engine/Build/BatchFiles/GetMSBuildPath.bat
--- a/Engine/Build/BatchFiles/GetMSBuildPath.bat
+++ b/Engine/Build/BatchFiles/GetMSBuildPath.bat
@@ -1,6 +1,6 @@
 @echo off
 
-rem ## Unreal Engine 4 utility script
+rem ## Unreal Engine utility script
 rem ## Copyright Epic Games, Inc. All Rights Reserved.
 rem ##
 rem ## This script determines the path to MSBuild necessary to compile C# tools for the current version of the engine.
diff --git a/Engine/Build/BatchFiles/Linux/ContainerBuildThirdParty.sh b/Engine/Build/BatchFiles/Linux/ContainerBuildThirdParty.sh
--- a/Engine/Build/BatchFiles/Linux/ContainerBuildThirdParty.sh
+++ b/Engine/Build/BatchFiles/Linux/ContainerBuildThirdParty.sh
@@ -1,6 +1,6 @@
 #!/bin/bash
 
-## Unreal Engine 4 Container Build script for Third Party Libs
+## Unreal Engine Container Build script for Third Party Libs
 ## Copyright Epic Games, Inc. All Rights Reserved.
 
 ## This script is expecting to exist in the UE4/Engine/Build/BatchFiles/Linux directory.  It will not work correctly
diff --git a/Engine/Build/BatchFiles/Linux/GenerateProjectFiles.sh b/Engine/Build/BatchFiles/Linux/GenerateProjectFiles.sh
--- a/Engine/Build/BatchFiles/Linux/GenerateProjectFiles.sh
+++ b/Engine/Build/BatchFiles/Linux/GenerateProjectFiles.sh
@@ -2,7 +2,7 @@
 # Copyright Epic Games, Inc. All Rights Reserved.
 
 echo
-echo Setting up Unreal Engine 5 project files...
+echo Setting up Unreal Engine project files...
 echo
 
 # If ran from somewhere other then the script location we'll have the full base path
@@ -25,7 +25,7 @@ source "$BASE_PATH/SetupEnvironment.sh" -dotnet "$BASE_PATH"
 # If this is a source drop of the engine make sure that the UnrealBuildTool is up-to-date
 if [ ! -f "$BASE_PATH/../../InstalledBuild.txt" ]; then
 	if [ -f "$BASE_PATH/../../../Source/Programs/UnrealBuildTool/UnrealBuildTool.csproj" ]; then
-		dotnet msbuild /restore /target:build /property:Configuration=Development /nologo "$BASE_PATH/../../../Source/Programs/UnrealBuildTool/UnrealBuildTool.csproj" /verbosity:quiet
+		dotnet build "$BASE_PATH/../../../Source/Programs/UnrealBuildTool/UnrealBuildTool.csproj" -c Development -v quiet
 		if [ $? -ne 0 ]; then
 			echo GenerateProjectFiles ERROR: Failed to build UnrealBuildTool
 			exit 1
diff --git a/Engine/Build/BatchFiles/Linux/RunMono.sh b/Engine/Build/BatchFiles/Linux/RunMono.sh
--- a/Engine/Build/BatchFiles/Linux/RunMono.sh
+++ b/Engine/Build/BatchFiles/Linux/RunMono.sh
@@ -1,6 +1,6 @@
 #!/bin/bash
 
-## Unreal Engine 4 Mono setup script
+## Unreal Engine Mono setup script
 ## Copyright Epic Games, Inc. All Rights Reserved.
 
 ## This script is expecting to exist in the UE4/Engine/Build/BatchFiles directory.  It will not work correctly
diff --git a/Engine/Build/BatchFiles/Linux/RunXBuild.sh b/Engine/Build/BatchFiles/Linux/RunXBuild.sh
--- a/Engine/Build/BatchFiles/Linux/RunXBuild.sh
+++ b/Engine/Build/BatchFiles/Linux/RunXBuild.sh
@@ -1,6 +1,6 @@
 #!/bin/bash
 
-## Unreal Engine 4 XBuild setup script
+## Unreal Engine XBuild setup script
 ## Copyright Epic Games, Inc. All Rights Reserved.
 
 ## This script is expecting to exist in the UE4/Engine/Build/BatchFiles directory.  It will not work correctly
diff --git a/Engine/Build/BatchFiles/MSBuild.bat b/Engine/Build/BatchFiles/MSBuild.bat
--- a/Engine/Build/BatchFiles/MSBuild.bat
+++ b/Engine/Build/BatchFiles/MSBuild.bat
@@ -1,6 +1,6 @@
 @echo off
 
-rem ## Unreal Engine 4 Visual Studio MSBuild Execution Script
+rem ## Unreal Engine Visual Studio MSBuild Execution Script
 rem ## Copyright Epic Games, Inc. All Rights Reserved.
 rem ##
 rem ## This script is expecting to exist in the UE4 Engine/Build/BatchFiles directory.  It will not work correctly
diff --git a/Engine/Build/BatchFiles/Mac/GenerateProjectFiles.sh b/Engine/Build/BatchFiles/Mac/GenerateProjectFiles.sh
--- a/Engine/Build/BatchFiles/Mac/GenerateProjectFiles.sh
+++ b/Engine/Build/BatchFiles/Mac/GenerateProjectFiles.sh
@@ -1,7 +1,7 @@
 #!/bin/bash
 # Copyright Epic Games, Inc. All Rights Reserved.
 echo
-echo Setting up Unreal Engine 5 project files...
+echo Setting up Unreal Engine project files...
 echo
 
 # If ran from somewhere other then the script location we'll have the full base path
@@ -23,7 +23,7 @@ fi
 source "$BASE_PATH/SetupEnvironment.sh" -dotnet "$BASE_PATH"
 # ensure UnrealBuildTool is up to date if the project file exists, but not if running from an installed build
 if [ -f "$BASE_PATH/../../../Source/Programs/UnrealBuildTool/UnrealBuildTool.csproj" -a ! -f "$BASE_PATH/../../../Build/InstalledBuild.txt" ]; then
-	dotnet msbuild /restore /target:build /property:Configuration=Development /nologo "$BASE_PATH/../../../Source/Programs/UnrealBuildTool/UnrealBuildTool.csproj" /verbosity:quiet
+	dotnet build "$BASE_PATH/../../../Source/Programs/UnrealBuildTool/UnrealBuildTool.csproj" -c Development -v quiet
 
 	if [ $? -ne 0 ]; then
 		echo GenerateProjectFiles ERROR: Failed to build UnrealBuildTool
diff --git a/Engine/Build/BatchFiles/Mac/RunMono.sh b/Engine/Build/BatchFiles/Mac/RunMono.sh
--- a/Engine/Build/BatchFiles/Mac/RunMono.sh
+++ b/Engine/Build/BatchFiles/Mac/RunMono.sh
@@ -1,6 +1,6 @@
 #!/bin/sh
 
-## Unreal Engine 4 Mono setup script
+## Unreal Engine Mono setup script
 ## Copyright Epic Games, Inc. All Rights Reserved.
 
 ## This script is expecting to exist in the UE4/Engine/Build/BatchFiles directory.  It will not work correctly
diff --git a/Engine/Build/BatchFiles/Mac/RunXBuild.sh b/Engine/Build/BatchFiles/Mac/RunXBuild.sh
--- a/Engine/Build/BatchFiles/Mac/RunXBuild.sh
+++ b/Engine/Build/BatchFiles/Mac/RunXBuild.sh
@@ -1,6 +1,6 @@
 #!/bin/sh
 
-## Unreal Engine 4 XBuild setup script
+## Unreal Engine XBuild setup script
 ## Copyright Epic Games, Inc. All Rights Reserved.
 
 ## This script is expecting to exist in the UE4/Engine/Build/BatchFiles directory.  It will not work correctly
diff --git a/Engine/Build/BatchFiles/Rebuild.bat b/Engine/Build/BatchFiles/Rebuild.bat
--- a/Engine/Build/BatchFiles/Rebuild.bat
+++ b/Engine/Build/BatchFiles/Rebuild.bat
@@ -1,21 +1,20 @@
 @echo off
 
-REM %1 is the game name
-REM %2 is the platform name
-REM %3 is the configuration name
-
-IF NOT EXIST "%~dp0\Clean.bat" GOTO Error_MissingCleanBatchFile
-
-call "%~dp0\Clean.bat" %*
+rem ## Unreal Engine code rebuild script
+rem ## Copyright Epic Games, Inc. All Rights Reserved.
+rem ##
+rem ## This script is expecting to exist in the Engine/Build/BatchFiles directory.  It will not work correctly
+rem ## if you copy it to a different location and run it.
+rem ##
+rem ##     %1 is the game name
+rem ##     %2 is the platform name
+rem ##     %3 is the configuration name
+rem ##     additional args are passed directly to UnrealBuildTool
 
 IF NOT EXIST "%~dp0\Build.bat" GOTO Error_MissingBuildBatchFile
 
-call "%~dp0\Build.bat" %*
-GOTO Exit
-
-:Error_MissingCleanBatchFile
-ECHO Clean.bat not found in "%~dp0"
-EXIT /B 999
+call "%~dp0\Build.bat" %* -Rebuild
+goto Exit
 
 :Error_MissingBuildBatchFile
 ECHO Build.bat not found in "%~dp0"
diff --git a/Engine/Build/BatchFiles/RunUAT.command b/Engine/Build/BatchFiles/RunUAT.command
--- a/Engine/Build/BatchFiles/RunUAT.command
+++ b/Engine/Build/BatchFiles/RunUAT.command
@@ -1,7 +1,7 @@
 #!/bin/sh
 ## Copyright Epic Games, Inc. All Rights Reserved.
 ##
-## Unreal Engine 4 AutomationTool setup script
+## Unreal Engine AutomationTool setup script
 ##
 ## This script uses the .command extenion so that is clickable in
 ## in the OSX Finder.  All it does is call RunUAT.sh which does
diff --git a/Engine/Build/BatchFiles/RunUBT.bat b/Engine/Build/BatchFiles/RunUBT.bat
--- a/Engine/Build/BatchFiles/RunUBT.bat
+++ b/Engine/Build/BatchFiles/RunUBT.bat
@@ -40,11 +40,11 @@ fc /b ..\Intermediate\Build\UnrealBuildToolFiles.txt ..\Intermediate\Build\Unrea
 if not errorlevel 1 goto SkipClean
 copy /y ..\Intermediate\Build\UnrealBuildToolFiles.txt ..\Intermediate\Build\UnrealBuildToolPrevFiles.txt >nul
 
-dotnet msbuild /target:clean /property:Configuration=Development /nologo Programs\UnrealBuildTool\UnrealBuildTool.csproj /verbosity:quiet
+dotnet clean Programs\UnrealBuildTool\UnrealBuildTool.csproj -c Development -v quiet
 
 :SkipClean
 echo Building UnrealBuildTool...
-dotnet msbuild /restore /target:build /property:Configuration=Development /property:Platform=AnyCPU /nologo Programs\UnrealBuildTool\UnrealBuildTool.csproj /verbosity:quiet
+dotnet build Programs\UnrealBuildTool\UnrealBuildTool.csproj -c Development -v quiet
 if errorlevel 1 goto Error_UBTCompileFailed
 
 rem ## Run UnrealBuildTool to generate Visual Studio solution and project files
diff --git a/Engine/Build/BatchFiles/RunUBT.sh b/Engine/Build/BatchFiles/RunUBT.sh
--- a/Engine/Build/BatchFiles/RunUBT.sh
+++ b/Engine/Build/BatchFiles/RunUBT.sh
@@ -29,7 +29,7 @@ fi
 if [[ "$*" != *-SkipUBTBuild* ]]; then
   if [ -f "$BASE_PATH/../../Source/Programs/UnrealBuildTool/UnrealBuildTool.csproj" ]; then 
     echo "Building UBT..."
-    dotnet msbuild /restore /target:build /property:Configuration=Development /nologo $BASE_PATH/../../Source/Programs/UnrealBuildTool/UnrealBuildTool.csproj /verbosity:quiet
+    dotnet build $BASE_PATH/../../Source/Programs/UnrealBuildTool/UnrealBuildTool.csproj -c Development -v quiet
 
     if [ $? -ne 0 ]; then
       echo RunUBT ERROR: Failed to build UnrealBuildTool
