Removing mono referenced from GitDependencies scripts

mono not used by this tool

#rnx

diff --git a/Engine/Build/BatchFiles/Linux/GitDependencies.sh b/Engine/Build/BatchFiles/Linux/GitDependencies.sh
--- a/Engine/Build/BatchFiles/Linux/GitDependencies.sh
+++ b/Engine/Build/BatchFiles/Linux/GitDependencies.sh
@@ -26,7 +26,7 @@ while : ; do
         echo "Result: $RESULT"
         # quit if not crashed
         [[ $RESULT -lt 129 ]] && break
-        echo "mono GitDependencies $ARGS crashed with return code $RESULT" >> GitDependencies.crash.log
+        echo "./Engine/Binaries/DotNET/GitDependencies/linux-x64/GitDependencies $ARGS crashed with return code $RESULT" >> GitDependencies.crash.log
 done
 
 exit $RESULT
diff --git a/Engine/Build/BatchFiles/Mac/GitDependencies.sh b/Engine/Build/BatchFiles/Mac/GitDependencies.sh
--- a/Engine/Build/BatchFiles/Mac/GitDependencies.sh
+++ b/Engine/Build/BatchFiles/Mac/GitDependencies.sh
@@ -10,11 +10,6 @@ fi
 
 cd "$(dirname "$SCRIPT_PATH")" && SCRIPT_PATH="`pwd`/$(basename "$SCRIPT_PATH")"
 
-if [ ! -f ../../../Binaries/DotNET/GitDependencies/osx-x64/GitDependencies ]; then
-	echo "Cannot find GitDependencies. This script should be placed in Engine/Build/BatchFiles/Mac."
-	exit 1
-fi 
-
 "$(dirname "$SCRIPT_PATH")/SetupDotnet.sh"
 
 cd ../../../..
