Correctly return errorlevel from Build.bat

#rnx

diff --git a/Engine/Build/BatchFiles/Build.bat b/Engine/Build/BatchFiles/Build.bat
--- a/Engine/Build/BatchFiles/Build.bat
+++ b/Engine/Build/BatchFiles/Build.bat
@@ -53,27 +53,20 @@ rem ## Run UBT
 if not exist %UBTPath% goto Error_UBTMissing
 echo Running UnrealBuildTool: dotnet %UBTPath% %*
 dotnet %UBTPath% %*
-goto Exit
+EXIT /B !ERRORLEVEL!
 
 :Error_BatchFileInWrongLocation
 echo ERROR: The batch file does not appear to be located in the Engine/Build/BatchFiles directory. This script must be run from within that directory.
-pause
-goto Exit
+EXIT /B 999
 
 :Error_NoDotnetSDK
-echo ERROR: Unable to find a install of Dotnet SDK.  Please make sure you have it installed and that `dotnet` is a globally available command.
-pause
-goto Exit
+echo ERROR: Unable to find an install of Dotnet SDK.  Please make sure you have it installed and that `dotnet` is a globally available command.
+EXIT /B 999
 
 :Error_UBTCompileFailed
 echo ERROR: Failed to build UnrealBuildTool.
-pause
-goto Exit
+EXIT /B 999
 
 :Error_UBTMissing
 echo ERROR: UnrealBuildTool.dll not found in %UBTPath%
-pause
-goto Exit
-
-:Exit
-
+EXIT /B 999
