Build.bat: Partially restore old logic where UBT was not compiled and only compile if it is missing, to prevent building mulitple times in parallel from visual studioand while UBT is running.

UBT should already be compiled via UGS (GenerateProjectFiles) or from VisualStudio, and while this could cause an out of date UBT to run this is a better stopgap fix then never building at all

#rnx

diff --git a/Engine/Build/BatchFiles/Build.bat b/Engine/Build/BatchFiles/Build.bat
--- a/Engine/Build/BatchFiles/Build.bat
+++ b/Engine/Build/BatchFiles/Build.bat
@@ -19,7 +19,7 @@ if not exist "%~dp0..\..\Source" goto Error_BatchFileInWrongLocation
 
 rem ## Change the CWD to /Engine/Source.  We always need to run UnrealBuildTool from /Engine/Source!
 pushd "%~dp0\..\..\Source"
-if not exist ..\Build\BatchFiles\Clean.bat goto Error_BatchFileInWrongLocation
+if not exist ..\Build\BatchFiles\Build.bat goto Error_BatchFileInWrongLocation
 
 set UBTPath="..\..\Engine\Binaries\DotNET\UnrealBuildTool\UnrealBuildTool.dll"
 
@@ -36,15 +36,20 @@ rem ## Compile UBT if the project file exists
 set ProjectFile="Programs\UnrealBuildTool\UnrealBuildTool.csproj"
 if not exist %ProjectFile% goto NoProjectFile
 
-rem ## If this script was called from Visual Studio 2022, build UBT with Visual Studio to prevent unnecessary rebuilds.
-if "%VisualStudioVersion%" GEQ "17.0" (
-	echo Building UnrealBuildTool with %VisualStudioEdition%...
-	"%VSAPPIDDIR%..\..\MSBuild\Current\Bin\MSBuild.exe" %ProjectFile% -t:Build -p:Configuration=Development -verbosity:quiet -noLogo
-	if errorlevel 1 goto Error_UBTCompileFailed
-) else (
-	echo Building UnrealBuildTool with dotnet...
-	dotnet build %ProjectFile% -c Development -v quiet
-	if errorlevel 1 goto Error_UBTCompileFailed
+rem ## Only build if UnrealBuildTool.dll is missing, as Visual Studio or GenerateProjectFiles should be building UnrealBuildTool
+rem ## Note: It is possible UnrealBuildTool will be out of date if the solution was generated with -NoDotNet or is VS2019
+rem ##       Historically this batch file did not compile UnrealBuildTool
+if not exist %UBTPath% (
+	rem ## If this script was called from Visual Studio 2022, build UBT with Visual Studio to prevent unnecessary rebuilds.
+	if "%VisualStudioVersion%" GEQ "17.0" (
+		echo Building UnrealBuildTool with %VisualStudioEdition%...
+		"%VSAPPIDDIR%..\..\MSBuild\Current\Bin\MSBuild.exe" %ProjectFile% -t:Build -p:Configuration=Development -verbosity:quiet -noLogo
+		if errorlevel 1 goto Error_UBTCompileFailed
+	) else (
+		echo Building UnrealBuildTool with dotnet...
+		dotnet build %ProjectFile% -c Development -v quiet
+		if errorlevel 1 goto Error_UBTCompileFailed
+	)
 )
 :NoProjectFile
 
