BatchFiles: Add output directory when caling dotnet build as running the scripts via a visual studio environment requires this, and for consistency. Via UDN

#rnx

diff --git a/Engine/Build/BatchFiles/Build.bat b/Engine/Build/BatchFiles/Build.bat
--- a/Engine/Build/BatchFiles/Build.bat
+++ b/Engine/Build/BatchFiles/Build.bat
@@ -21,6 +21,7 @@ rem ## Change the CWD to /Engine/Source.  We always need to run UnrealBuildTool
 pushd "%~dp0\..\..\Source"
 if not exist ..\Build\BatchFiles\Build.bat goto Error_BatchFileInWrongLocation
 
+set UBTDir="..\..\Engine\Binaries\DotNET\UnrealBuildTool"
 set UBTPath="..\..\Engine\Binaries\DotNET\UnrealBuildTool\UnrealBuildTool.dll"
 
 rem ## If this is an installed build, we don't need to rebuild UBT. Go straight to building.
@@ -47,7 +48,7 @@ if not exist %UBTPath% (
 		if errorlevel 1 goto Error_UBTCompileFailed
 	) else (
 		echo Building UnrealBuildTool with dotnet...
-		dotnet build %ProjectFile% -c Development -v quiet
+		dotnet build %ProjectFile% -o %UBTDir% -c Development -v quiet
 		if errorlevel 1 goto Error_UBTCompileFailed
 	)
 )
diff --git a/Engine/Build/BatchFiles/BuildUAT.bat b/Engine/Build/BatchFiles/BuildUAT.bat
--- a/Engine/Build/BatchFiles/BuildUAT.bat
+++ b/Engine/Build/BatchFiles/BuildUAT.bat
@@ -60,7 +60,7 @@ call "%~dp0GetDotnetPath.bat"
 if errorlevel 1 goto Error_NoDotnetSDK
 
 echo Building AutomationTool...
-dotnet build Programs\AutomationTool\AutomationTool.csproj -c Development -v %MSBUILD_LOGLEVEL%
+dotnet build Programs\AutomationTool\AutomationTool.csproj -o ..\Binaries\DotNET\AutomationTool -c Development -v %MSBUILD_LOGLEVEL%
 if errorlevel 1 goto Error_UATCompileFailed
 
 rem record input files - regardless of how we got here, these are now our point of reference
diff --git a/Engine/Build/BatchFiles/BuildUAT.sh b/Engine/Build/BatchFiles/BuildUAT.sh
--- a/Engine/Build/BatchFiles/BuildUAT.sh
+++ b/Engine/Build/BatchFiles/BuildUAT.sh
@@ -77,7 +77,7 @@ if [ $PERFORM_REBUILD -eq 1 ]; then
     source "$SCRIPT_DIR/Linux/SetupEnvironment.sh" $EnvironmentType "$SCRIPT_DIR/Linux"
   fi
 
-  dotnet build Programs/AutomationTool/AutomationTool.csproj -c Development -v $MSBuild_Verbosity
+  dotnet build Programs/AutomationTool/AutomationTool.csproj -o ../Binaries/DotNET/AutomationTool -c Development -v $MSBuild_Verbosity
   if [ $? -ne 0 ]; then
     echo "Compilation failed"
     popd >/dev/null
diff --git a/Engine/Build/BatchFiles/GenerateProjectFiles.bat b/Engine/Build/BatchFiles/GenerateProjectFiles.bat
--- a/Engine/Build/BatchFiles/GenerateProjectFiles.bat
+++ b/Engine/Build/BatchFiles/GenerateProjectFiles.bat
@@ -52,7 +52,7 @@ dotnet clean Programs\UnrealBuildTool\UnrealBuildTool.csproj -c Development -v q
 
 :SkipClean
 echo Building UnrealBuildTool...
-dotnet build Programs\UnrealBuildTool\UnrealBuildTool.csproj -c Development -v quiet
+dotnet build Programs\UnrealBuildTool\UnrealBuildTool.csproj -o ..\Binaries\DotNET\UnrealBuildTool  -c Development -v quiet
 
 if errorlevel 1 goto Error_UBTCompileFailed
 
diff --git a/Engine/Build/BatchFiles/Linux/Build.sh b/Engine/Build/BatchFiles/Linux/Build.sh
--- a/Engine/Build/BatchFiles/Linux/Build.sh
+++ b/Engine/Build/BatchFiles/Linux/Build.sh
@@ -16,7 +16,7 @@ if [ ! -f Engine/Build/InstalledBuild.txt ]; then
 	for i in "$@" ; do
 	# First make sure that the UnrealBuildTool is up-to-date
 	if [[ $i == "-buildubt" ]] ; then
-		if ! dotnet build Engine/Source/Programs/UnrealBuildTool/UnrealBuildTool.csproj -c Development -v quiet; then
+		if ! dotnet build Engine/Source/Programs/UnrealBuildTool/UnrealBuildTool.csproj -o Engine/Binaries/DotNET/UnrealBuildTool -c Development -v quiet; then
 			echo "Failed to build the build tool (UnrealBuildTool)"
 			exit 1
 		fi
diff --git a/Engine/Build/BatchFiles/Linux/GenerateProjectFiles.sh b/Engine/Build/BatchFiles/Linux/GenerateProjectFiles.sh
--- a/Engine/Build/BatchFiles/Linux/GenerateProjectFiles.sh
+++ b/Engine/Build/BatchFiles/Linux/GenerateProjectFiles.sh
@@ -25,7 +25,7 @@ source "$BASE_PATH/SetupEnvironment.sh" -dotnet "$BASE_PATH"
 # If this is a source drop of the engine make sure that the UnrealBuildTool is up-to-date
 if [ ! -f "$BASE_PATH/../../InstalledBuild.txt" ]; then
 	if [ -f "$BASE_PATH/../../../Source/Programs/UnrealBuildTool/UnrealBuildTool.csproj" ]; then
-		dotnet build "$BASE_PATH/../../../Source/Programs/UnrealBuildTool/UnrealBuildTool.csproj" -c Development -v quiet
+		dotnet build "$BASE_PATH/../../../Source/Programs/UnrealBuildTool/UnrealBuildTool.csproj" -o "$BASE_PATH/../../../Binaries/DotNET/UnrealBuildTool" -c Development -v quiet
 		if [ $? -ne 0 ]; then
 			echo GenerateProjectFiles ERROR: Failed to build UnrealBuildTool
 			exit 1
diff --git a/Engine/Build/BatchFiles/Mac/Build.sh b/Engine/Build/BatchFiles/Mac/Build.sh
--- a/Engine/Build/BatchFiles/Mac/Build.sh
+++ b/Engine/Build/BatchFiles/Mac/Build.sh
@@ -14,7 +14,7 @@ if [ ! -f Engine/Build/InstalledBuild.txt ]; then
 	for i in "$@" ; do
 	# First make sure that the UnrealBuildTool is up-to-date
 	if [[ $i == "-buildubt" ]] ; then
-		if ! dotnet build Engine/Source/Programs/UnrealBuildTool/UnrealBuildTool.csproj -c Development -v quiet; then
+		if ! dotnet build Engine/Source/Programs/UnrealBuildTool/UnrealBuildTool.csproj -o Engine/Binaries/DotNET/UnrealBuildTool -c Development -v quiet; then
 			echo "Failed to build the build tool (UnrealBuildTool)"
 			exit 1
 		fi
diff --git a/Engine/Build/BatchFiles/Mac/GenerateProjectFiles.sh b/Engine/Build/BatchFiles/Mac/GenerateProjectFiles.sh
--- a/Engine/Build/BatchFiles/Mac/GenerateProjectFiles.sh
+++ b/Engine/Build/BatchFiles/Mac/GenerateProjectFiles.sh
@@ -23,7 +23,7 @@ fi
 source "$BASE_PATH/SetupEnvironment.sh" -dotnet "$BASE_PATH"
 # ensure UnrealBuildTool is up to date if the project file exists, but not if running from an installed build
 if [ -f "$BASE_PATH/../../../Source/Programs/UnrealBuildTool/UnrealBuildTool.csproj" -a ! -f "$BASE_PATH/../../../Build/InstalledBuild.txt" ]; then
-	dotnet build "$BASE_PATH/../../../Source/Programs/UnrealBuildTool/UnrealBuildTool.csproj" -c Development -v quiet
+	dotnet build "$BASE_PATH/../../../Source/Programs/UnrealBuildTool/UnrealBuildTool.csproj" -o "$BASE_PATH/../../../Binaries/DotNET/UnrealBuildTool" -c Development -v quiet
 
 	if [ $? -ne 0 ]; then
 		echo GenerateProjectFiles ERROR: Failed to build UnrealBuildTool
diff --git a/Engine/Build/BatchFiles/Mac/XcodeBuild.sh b/Engine/Build/BatchFiles/Mac/XcodeBuild.sh
--- a/Engine/Build/BatchFiles/Mac/XcodeBuild.sh
+++ b/Engine/Build/BatchFiles/Mac/XcodeBuild.sh
@@ -17,7 +17,7 @@ unset TARGETNAME
 
 # If this is a source drop of the engine make sure that the UnrealBuildTool is up-to-date
 if [ ! -f Engine/Build/InstalledBuild.txt ]; then
-	dotnet build Engine/Source/Programs/UnrealBuildTool/UnrealBuildTool.csproj -c Development -v quiet
+	dotnet build Engine/Source/Programs/UnrealBuildTool/UnrealBuildTool.csproj -o Engine/Binaries/DotNET/UnrealBuildTool -c Development -v quiet
 
 	if [ $? -ne 0 ]; then
 		echo "Failed to build the build tool (UnrealBuildTool)"
diff --git a/Engine/Build/BatchFiles/RunUBT.bat b/Engine/Build/BatchFiles/RunUBT.bat
--- a/Engine/Build/BatchFiles/RunUBT.bat
+++ b/Engine/Build/BatchFiles/RunUBT.bat
@@ -44,7 +44,7 @@ dotnet clean Programs\UnrealBuildTool\UnrealBuildTool.csproj -c Development -v q
 
 :SkipClean
 echo Building UnrealBuildTool...
-dotnet build Programs\UnrealBuildTool\UnrealBuildTool.csproj -c Development -v quiet
+dotnet build Programs\UnrealBuildTool\UnrealBuildTool.csproj -o ..\Binaries\DotNET\UnrealBuildTool -c Development -v quiet
 if errorlevel 1 goto Error_UBTCompileFailed
 
 rem ## Run UnrealBuildTool to generate Visual Studio solution and project files
diff --git a/Engine/Build/BatchFiles/RunUBT.sh b/Engine/Build/BatchFiles/RunUBT.sh
--- a/Engine/Build/BatchFiles/RunUBT.sh
+++ b/Engine/Build/BatchFiles/RunUBT.sh
@@ -29,7 +29,7 @@ fi
 if [[ "$*" != *-SkipUBTBuild* ]]; then
   if [ -f "$BASE_PATH/../../Source/Programs/UnrealBuildTool/UnrealBuildTool.csproj" ]; then 
     echo "Building UBT..."
-    dotnet build $BASE_PATH/../../Source/Programs/UnrealBuildTool/UnrealBuildTool.csproj -c Development -v quiet
+    dotnet build $BASE_PATH/../../Source/Programs/UnrealBuildTool/UnrealBuildTool.csproj -o $BASE_PATH/../../Binaries/DotNET/UnrealBuildTool -c Development -v quiet
 
     if [ $? -ne 0 ]; then
       echo RunUBT ERROR: Failed to build UnrealBuildTool
