UAT: Add support for Windows on ARM64 through RunUAT.bat.

diff --git a/Engine/Build/BatchFiles/GetDotnetPath.bat b/Engine/Build/BatchFiles/GetDotnetPath.bat
--- a/Engine/Build/BatchFiles/GetDotnetPath.bat
+++ b/Engine/Build/BatchFiles/GetDotnetPath.bat
@@ -8,9 +8,15 @@ rem ## This script verifies that dotnet sdk is installed and a new enough SDK is
 rem if UE_USE_SYSTEM_DOTNET we assume a installed dotnet is present
 if "%UE_USE_SYSTEM_DOTNET%" == "1" goto verify_dotnet
 
+rem figure out which architecture to use
+set UE_DOTNET_ARCH=windows
+if "%PROCESSOR_ARCHITECTURE%" == "ARM64" (
+	set UE_DOTNET_ARCH=win-arm64
+)
+
 rem add the dotnet sdk in the repo as the current dotnet sdk
 set UE_DOTNET_VERSION=6.0.302
-set UE_DOTNET_DIR=%~dp0..\..\Binaries\ThirdParty\DotNet\%UE_DOTNET_VERSION%\windows
+set UE_DOTNET_DIR=%~dp0..\..\Binaries\ThirdParty\DotNet\%UE_DOTNET_VERSION%\%UE_DOTNET_ARCH%
 set PATH=%UE_DOTNET_DIR%;%PATH%
 set DOTNET_ROOT=%UE_DOTNET_DIR%
 
diff --git a/Engine/Source/Programs/Shared/EpicGames.Build/Unreal.cs b/Engine/Source/Programs/Shared/EpicGames.Build/Unreal.cs
--- a/Engine/Source/Programs/Shared/EpicGames.Build/Unreal.cs
+++ b/Engine/Source/Programs/Shared/EpicGames.Build/Unreal.cs
@@ -93,7 +93,15 @@ namespace UnrealBuildBase
 			string HostDotNetDirectoryName;
 			switch (HostPlatform)
 			{
-				case RuntimePlatform.Type.Windows: HostDotNetDirectoryName = "windows"; break;
+				case RuntimePlatform.Type.Windows:
+					{
+						HostDotNetDirectoryName = "windows";
+						if (RuntimeInformation.ProcessArchitecture == Architecture.Arm64)
+						{
+							HostDotNetDirectoryName = "win-arm64";
+						}
+						break;
+					}
 				case RuntimePlatform.Type.Mac:
 					{
 						HostDotNetDirectoryName = "mac-x64";
