BatchFiles: Clean up so non-ide UBT builds all call the same script and fix up some dependency scanning

#rnx

diff --git a/Engine/Build/BatchFiles/BuildUAT.bat b/Engine/Build/BatchFiles/BuildUAT.bat
--- a/Engine/Build/BatchFiles/BuildUAT.bat
+++ b/Engine/Build/BatchFiles/BuildUAT.bat
@@ -21,18 +21,48 @@ rem find ".cs" files to only lines that match those names - excludes lines that
 md ..\Intermediate\Build >nul 2>nul
 
 dir /s^
- Programs\Shared\EpicGames.Core\*.cs^
- Programs\Shared\EpicGames.Core\*.csproj^
  Programs\Shared\EpicGames.Build\*.cs^
  Programs\Shared\EpicGames.Build\*.csproj^
+ Programs\Shared\EpicGames.Core\*.cs^
+ Programs\Shared\EpicGames.Core\*.csproj^
+ Programs\Shared\EpicGames.IoHash\*.cs^
+ Programs\Shared\EpicGames.IoHash\*.csproj^
  Programs\Shared\EpicGames.MsBuild\*.cs^
  Programs\Shared\EpicGames.MsBuild\*.csproj^
+ Programs\Shared\EpicGames.Serialization\*.cs^
+ Programs\Shared\EpicGames.Serialization\*.csproj^
  Programs\Shared\EpicGames.UHT\*.cs^
  Programs\Shared\EpicGames.UHT\*.csproj^
  Programs\UnrealBuildTool\*.cs^
  Programs\UnrealBuildTool\*.csproj^
  | find ".cs" >..\Intermediate\Build\AutomationToolFiles.txt
 
+if not exist ..\Platforms goto NoPlatforms
+for /d %%D in (..\Platforms\*) do (
+	if exist %%D\Source\Programs\UnrealBuildTool (
+		dir /s^
+		 %%D\Source\Programs\AutomationTool\*.cs^
+		 %%D\Source\Programs\AutomationTool\*.csproj^
+		 %%D\Source\Programs\UnrealBuildTool\*.cs^
+		 %%D\Source\Programs\UnrealBuildTool\*.csproj^
+		 | find ".cs" >> ..\Intermediate\Build\AutomationToolFiles.txt
+	)
+)
+:NoPlatforms
+
+if not exist ..\Restricted goto NoRestricted
+for /d %%D in (..\Restricted\*) do (
+	if exist %%D\Source\Programs\UnrealBuildTool (
+		dir /s^
+		 %%D\Source\Programs\AutomationTool\*.cs^
+		 %%D\Source\Programs\AutomationTool\*.csproj^
+		 %%D\Source\Programs\UnrealBuildTool\*.cs^
+		 %%D\Source\Programs\UnrealBuildTool\*.csproj^
+		 | find ".cs" >> ..\Intermediate\Build\AutomationToolFiles.txt
+	)
+)
+:NoRestricted
+
 rem note: no /s
 dir ^
  Programs\Shared\MetaData.cs^
@@ -40,8 +70,6 @@ dir ^
  Programs\AutomationTool\*.csproj^
  | find ".cs" >>..\Intermediate\Build\AutomationToolFiles.txt
 
-if not exist ..\Binaries\DotNET\AutomationTool\AutomationTool.dll goto Build_AutomationTool
-
 set MSBUILD_LOGLEVEL=%1
 if not defined %MSBUILD_LOGLEVEL set MSBUILD_LOGLEVEL=quiet
 
@@ -50,6 +78,7 @@ if not defined %ARGUMENT goto Check_UpToDate
 if /I "%ARGUMENT%" == "FORCE" goto Build_AutomationTool
 
 :Check_UpToDate
+if not exist ..\Binaries\DotNET\AutomationTool\AutomationTool.dll goto Build_AutomationTool
 set RUNUAT_EXITCODE=0
 rem per https://ss64.com/nt/fc.html using redirection syntax rather than errorlevel, based on observed inconsistent results from this function
 fc ..\Intermediate\Build\AutomationToolFiles.txt ..\Intermediate\Build\AutomationToolPrevFiles.txt >nul && goto Exit
@@ -60,7 +89,8 @@ call "%~dp0GetDotnetPath.bat"
 if errorlevel 1 goto Error_NoDotnetSDK
 
 echo Building AutomationTool...
-dotnet build Programs\AutomationTool\AutomationTool.csproj -c Development -v %MSBUILD_LOGLEVEL%
+dotnet clean Programs\AutomationTool\AutomationTool.csproj -c Development -v %MSBUILD_LOGLEVEL% -o ..\Binaries\DotNET\AutomationTool --nologo
+dotnet build Programs\AutomationTool\AutomationTool.csproj -c Development -v %MSBUILD_LOGLEVEL% -o ..\Binaries\DotNET\AutomationTool
 if errorlevel 1 goto Error_UATCompileFailed
 
 rem record input files - regardless of how we got here, these are now our point of reference
diff --git a/Engine/Build/BatchFiles/BuildUAT.sh b/Engine/Build/BatchFiles/BuildUAT.sh
--- a/Engine/Build/BatchFiles/BuildUAT.sh
+++ b/Engine/Build/BatchFiles/BuildUAT.sh
@@ -19,7 +19,7 @@ if [ ! -f ../Build/BatchFiles/BuildUAT.sh ]; then
   exit 1
 fi
 
-MSBuild_Verbosity="${1:quiet}"
+MSBuild_Verbosity="${1:-quiet}"
 
 # Check to see if the files in the AutomationTool, EpicGames.Build, EpicGames.Core, or UnrealBuildTool
 # directory have changed.
@@ -37,7 +37,17 @@ elif [ ! -f ../Binaries/DotNET/AutomationTool/AutomationTool.dll ]; then
 
 elif [ -f ../Intermediate/Build/AutomationToolLastBuildTime ]; then
   UPDATED_DEP_FILES="$(find \
-    Programs/Shared/EpicGames.Core Programs/Shared/EpicGames.Build Programs/Shared/EpicGames.MsBuild Programs/Shared/EpicGames.UHT Programs/UnrealBuildTool \
+    Programs/Shared/EpicGames.Build \
+    Programs/Shared/EpicGames.Core \
+    Programs/Shared/EpicGames.IoHash \
+    Programs/Shared/EpicGames.MsBuild \
+    Programs/Shared/EpicGames.Serialization \
+    Programs/Shared/EpicGames.UHT \
+    Programs/UnrealBuildTool \
+    ../Restricted/**/Source/Programs/UnrealBuildTool \
+    ../Platforms/*/Source/Programs/UnrealBuildTool \
+    ../Restricted/**/Source/Programs/AutomationTool \
+    ../Platforms/*/Source/Programs/AutomationTool \
     -type f \
     \( -iname \*.cs -or -iname \*.csproj \) \
     -newer ../Intermediate/Build/AutomationToolLastBuildTime)"
@@ -74,7 +84,7 @@ if [ $PERFORM_REBUILD -eq 1 ]; then
 
   if [ "$(uname)" = "Linux" ]; then
     # Setup Environment
-    source "$SCRIPT_DIR/Linux/SetupEnvironment.sh" $EnvironmentType "$SCRIPT_DIR/Linux"
+    source "$SCRIPT_DIR/Linux/SetupEnvironment.sh" -dotnet "$SCRIPT_DIR/Linux"
   fi
 
   dotnet build Programs/AutomationTool/AutomationTool.csproj -c Development -v $MSBuild_Verbosity
diff --git a/Engine/Build/BatchFiles/BuildUBT.bat b/Engine/Build/BatchFiles/BuildUBT.bat
new file mode 100755
--- /dev/null
+++ b/Engine/Build/BatchFiles/BuildUBT.bat
@@ -0,0 +1,121 @@
+@echo off
+
+rem ## Unreal Engine UnrealBuildTool build script
+rem ## Copyright Epic Games, Inc. All Rights Reserved.
+rem ##
+rem ## This script is expecting to exist in the Engine/Build/BatchFiles directory.  It will not work correctly
+rem ## if you copy it to a different location and run it.
+
+setlocal
+
+rem ## First, make sure the batch file exists in the folder we expect it to.  This is necessary in order to
+rem ## verify that our relative path to the /Engine/Source directory is correct
+if not exist "%~dp0..\..\Source" goto Error_BatchFileInWrongLocation
+
+rem ## Change the CWD to /Engine/Source.
+pushd "%~dp0..\..\Source"
+if not exist ..\Build\BatchFiles\BuildUBT.bat goto Error_BatchFileInWrongLocation
+
+rem Check to see if the files in the UnrealBuildTool or Shared project directory have changed.
+rem find ".cs" files to only lines that match those names - excludes lines that will change for uninteresting reasons, like free space
+md ..\Intermediate\Build >nul 2>nul
+
+dir /s^
+ Programs\Shared\EpicGames.Build\*.cs^
+ Programs\Shared\EpicGames.Build\*.csproj^
+ Programs\Shared\EpicGames.Core\*.cs^
+ Programs\Shared\EpicGames.Core\*.csproj^
+ Programs\Shared\EpicGames.IoHash\*.cs^
+ Programs\Shared\EpicGames.IoHash\*.csproj^
+ Programs\Shared\EpicGames.MsBuild\*.cs^
+ Programs\Shared\EpicGames.MsBuild\*.csproj^
+ Programs\Shared\EpicGames.Serialization\*.cs^
+ Programs\Shared\EpicGames.Serialization\*.csproj^
+ Programs\Shared\EpicGames.UHT\*.cs^
+ Programs\Shared\EpicGames.UHT\*.csproj^
+ Programs\UnrealBuildTool\*.cs^
+ Programs\UnrealBuildTool\*.csproj^
+ | find ".cs" > ..\Intermediate\Build\UnrealBuildToolFiles.txt
+
+if not exist ..\Platforms goto NoPlatforms
+for /d %%D in (..\Platforms\*) do (
+	if exist %%D\Source\Programs\UnrealBuildTool (
+		dir /s^
+		 %%D\Source\Programs\UnrealBuildTool\*.cs^
+		 %%D\Source\Programs\UnrealBuildTool\*.csproj^
+		 | find ".cs" >> ..\Intermediate\Build\UnrealBuildToolFiles.txt
+	)
+)
+:NoPlatforms
+
+if not exist ..\Restricted goto NoRestricted
+for /d %%D in (..\Restricted\*) do (
+	if exist %%D\Source\Programs\UnrealBuildTool (
+		dir /s^
+		 %%D\Source\Programs\UnrealBuildTool\*.cs^
+		 %%D\Source\Programs\UnrealBuildTool\*.csproj^
+		 | find ".cs" >> ..\Intermediate\Build\UnrealBuildToolFiles.txt
+	)
+)
+:NoRestricted
+
+rem note: no /s
+dir ^
+ Programs\Shared\MetaData.cs^
+ | find ".cs" >>..\Intermediate\Build\UnrealBuildToolFiles.txt
+
+set MSBUILD_LOGLEVEL=%1
+if not defined %MSBUILD_LOGLEVEL set MSBUILD_LOGLEVEL=quiet
+
+set ARGUMENT=%2
+if not defined %ARGUMENT goto Check_UpToDate
+if /I "%ARGUMENT%" == "FORCE" goto Build_UnrealBuildTool
+
+:Check_UpToDate
+if not exist ..\Binaries\DotNET\UnrealBuildTool\UnrealBuildTool.dll goto Build_UnrealBuildTool
+set RUNUBT_EXITCODE=0
+rem per https://ss64.com/nt/fc.html using redirection syntax rather than errorlevel, based on observed inconsistent results from this function
+fc ..\Intermediate\Build\UnrealBuildToolFiles.txt ..\Intermediate\Build\UnrealBuildToolPrevFiles.txt >nul && goto Exit
+
+:Build_UnrealBuildTool
+rem ## Verify that dotnet is present
+call "%~dp0GetDotnetPath.bat"
+if errorlevel 1 goto Error_NoDotnetSDK
+
+:SkipClean
+echo Building UnrealBuildTool...
+dotnet clean Programs\UnrealBuildTool\UnrealBuildTool.csproj -c Development -v %MSBUILD_LOGLEVEL% -o ..\Binaries\DotNET\UnrealBuildTool --nologo
+dotnet build Programs\UnrealBuildTool\UnrealBuildTool.csproj -c Development -v %MSBUILD_LOGLEVEL% -o ..\Binaries\DotNET\UnrealBuildTool
+if errorlevel 1 goto Error_UBTCompileFailed
+
+rem record input files - regardless of how we got here, these are now our point of reference
+copy /y ..\Intermediate\Build\UnrealBuildToolFiles.txt ..\Intermediate\Build\UnrealBuildToolPrevFiles.txt >nul
+
+goto Exit
+
+
+:Error_BatchFileInWrongLocation
+echo.
+echo BuildUBT ERROR: The batch file does not appear to be located in the /Engine/Build/BatchFiles directory.  This script must be run from within that directory.
+echo.
+set RUNUBT_EXITCODE=1
+goto Exit
+
+:Error_NoDotnetSDK
+echo.
+echo RunUBT ERROR: Unable to find a install of Dotnet SDK.  Please make sure you have it installed and that `dotnet` is a globally available command.
+echo.
+set RUNUBT_EXITCODE=1
+goto Exit
+
+:Error_UBTCompileFailed
+echo.
+echo RunUBT ERROR: UnrealBuildTool failed to compile.
+echo.
+set RUNUBT_EXITCODE=1
+goto Exit
+
+:Exit
+rem ## Restore original CWD in case we change it
+popd
+exit /B %RUNUBT_EXITCODE%
diff --git a/Engine/Build/BatchFiles/BuildUBT.sh b/Engine/Build/BatchFiles/BuildUBT.sh
new file mode 100755
--- /dev/null
+++ b/Engine/Build/BatchFiles/BuildUBT.sh
@@ -0,0 +1,101 @@
+#!/bin/bash
+
+## Unreal Engine UnrealBuildTool build script
+## Copyright Epic Games, Inc. All Rights Reserved
+##
+## This script is expecting to exist in the Engine/Build/BatchFiles directory.  It will not work correctly
+## if you copy it to a different location and run it.
+
+## First, make sure the batch file exists in the folder we expect it to.  This is necessary in order to
+## verify that our relative path to the /Engine/Source directory is correct
+SCRIPT_DIR=$(cd "`dirname "$0"`" && pwd)
+pushd "$SCRIPT_DIR/../../Source" >/dev/null
+
+if [ ! -f ../Build/BatchFiles/BuildUBT.sh ]; then
+  echo
+  echo "BuildUBT ERROR: The script does not appear to be located in the /Engine/Build/BatchFiles directory.  This script must be run from within that directory."
+  echo
+  popd >/dev/null
+  exit 1
+fi
+
+MSBuild_Verbosity="${1:-quiet}"
+
+# Check to see if the files in the UnrealBuildTool, EpicGames.Build, EpicGames.Core, or UnrealBuildTool
+# directory have changed.
+
+mkdir -p ../Intermediate/Build
+
+PERFORM_REBUILD=0
+if [ "$1" == "FORCE" ]; then
+  PERFORM_REBUILD=1
+  echo "Rebuilding: build requested"
+
+elif [ ! -f ../Binaries/DotNET/UnrealBuildTool/UnrealBuildTool.dll ]; then
+  PERFORM_REBUILD=1
+  echo "Rebuilding: UnrealBuildTool assembly not found"
+
+elif [ -f ../Intermediate/Build/UnrealBuildToolLastBuildTime ]; then
+  UPDATED_DEP_FILES="$(find \
+    Programs/Shared/EpicGames.Build \
+    Programs/Shared/EpicGames.Core \
+    Programs/Shared/EpicGames.IoHash \
+    Programs/Shared/EpicGames.MsBuild \
+    Programs/Shared/EpicGames.Serialization \
+    Programs/Shared/EpicGames.UHT \
+    Programs/UnrealBuildTool \
+    ../Restricted/**/Source/Programs/UnrealBuildTool \
+    ../Platforms/*/Source/Programs/UnrealBuildTool \
+    -type f \
+    \( -iname \*.cs -or -iname \*.csproj \) \
+    -newer ../Intermediate/Build/UnrealBuildToolLastBuildTime)"
+  
+  if [ -n "$UPDATED_DEP_FILES" ]; then
+    PERFORM_REBUILD=1
+    echo "Rebuilding: found updated files:"
+    echo "$UPDATED_DEP_FILES"
+  fi
+
+  UPDATED_DEP_FILES="$(find \
+    Programs/Shared \
+    -maxdepth 1 \
+    -type f \
+    \( -iname \*.cs -or -iname \*.csproj \) \
+    -newer ../Intermediate/Build/UnrealBuildToolLastBuildTime)"
+  if [ -n "$UPDATED_DEP_FILES" ]; then
+    PERFORM_REBUILD=1
+    echo "Rebuilding: Found updated files:"
+    echo "$UPDATED_DEP_FILES"
+  fi
+
+else
+  PERFORM_REBUILD=1
+  echo "Rebuilding: No record of previous builld"
+fi
+
+if [ $PERFORM_REBUILD -eq 1 ]; then
+  if [ "$(uname)" = "Darwin" ]; then
+    # Setup Environment
+    source "$SCRIPT_DIR/Mac/SetupEnvironment.sh" -dotnet "$SCRIPT_DIR/Mac"
+  fi
+
+  if [ "$(uname)" = "Linux" ]; then
+    # Setup Environment
+    source "$SCRIPT_DIR/Linux/SetupEnvironment.sh" -dotnet "$SCRIPT_DIR/Linux"
+  fi
+
+  dotnet build Programs/UnrealBuildTool/UnrealBuildTool.csproj -c Development -v $MSBuild_Verbosity
+  if [ $? -ne 0 ]; then
+    echo "Compilation failed"
+    popd >/dev/null
+    exit 1
+  fi
+
+  touch ../Intermediate/Build/UnrealBuildToolLastBuildTime
+else
+  echo "UnrealBuildTool.dll is up to date"
+fi
+
+popd >/dev/null
+exit 0
+
diff --git a/Engine/Build/BatchFiles/GenerateProjectFiles.bat b/Engine/Build/BatchFiles/GenerateProjectFiles.bat
--- a/Engine/Build/BatchFiles/GenerateProjectFiles.bat
+++ b/Engine/Build/BatchFiles/GenerateProjectFiles.bat
@@ -28,32 +28,8 @@ call "%~dp0GetDotnetPath.bat"
 if errorlevel 1 goto Error_NoDotnetSDK
 REM ## Skip msbuild detection if using dotnet as this is done for us by dotnet-cli
 
-rem Check to see if the files in the UBT directory have changed. We conditionally include platform files from the .csproj file, but MSBuild doesn't recognize the dependency when new files are added. 
-md ..\Intermediate\Build >nul 2>nul
-dir /s /b Programs\UnrealBuildTool\*.cs >..\Intermediate\Build\UnrealBuildToolFiles.txt
-
-if not exist ..\Platforms goto NoPlatforms
-for /d %%D in (..\Platforms\*) do (
-	if exist %%D\Source\Programs\UnrealBuildTool dir /s /b %%D\Source\Programs\UnrealBuildTool\*.cs 2> nul >> ..\Intermediate\Build\UnrealBuildToolFiles.txt
-)
-:NoPlatforms
-
-if not exist ..\Restricted goto NoRestricted
-for /d %%D in (..\Restricted\*) do (
-	if exist %%D\Source\Programs\UnrealBuildTool dir /s /b %%D\Source\Programs\UnrealBuildTool\*.cs 2> nul >> ..\Intermediate\Build\UnrealBuildToolFiles.txt
-)
-:NoRestricted
-
-fc /b ..\Intermediate\Build\UnrealBuildToolFiles.txt ..\Intermediate\Build\UnrealBuildToolPrevFiles.txt >nul 2>nul
-if not errorlevel 1 goto SkipClean
-copy /y ..\Intermediate\Build\UnrealBuildToolFiles.txt ..\Intermediate\Build\UnrealBuildToolPrevFiles.txt >nul
-
-dotnet clean Programs\UnrealBuildTool\UnrealBuildTool.csproj -c Development -v quiet
-
-:SkipClean
-echo Building UnrealBuildTool...
-dotnet build Programs\UnrealBuildTool\UnrealBuildTool.csproj -c Development -v quiet
-
+rem ## Build UnrealBuildTool if necessary
+call "%~dp0BuildUBT.bat"
 if errorlevel 1 goto Error_UBTCompileFailed
 
 rem ## Run UnrealBuildTool to generate Visual Studio solution and project files
diff --git a/Engine/Build/BatchFiles/GetDotnetPath.bat b/Engine/Build/BatchFiles/GetDotnetPath.bat
--- a/Engine/Build/BatchFiles/GetDotnetPath.bat
+++ b/Engine/Build/BatchFiles/GetDotnetPath.bat
@@ -8,6 +8,9 @@ rem ## This script verifies that dotnet sdk is installed and a new enough SDK is
 rem if UE_USE_SYSTEM_DOTNET we assume a installed dotnet is present
 if "%UE_USE_SYSTEM_DOTNET%" == "1" goto verify_dotnet
 
+rem if UE_DOTNET_VERSION is already set we assume this script has already run
+if "%UE_DOTNET_VERSION%" == "6.0.302" goto Succeeded
+
 rem figure out which architecture to use
 set UE_DOTNET_ARCH=windows
 if "%PROCESSOR_ARCHITECTURE%" == "ARM64" (
diff --git a/Engine/Build/BatchFiles/Linux/GenerateProjectFiles.sh b/Engine/Build/BatchFiles/Linux/GenerateProjectFiles.sh
--- a/Engine/Build/BatchFiles/Linux/GenerateProjectFiles.sh
+++ b/Engine/Build/BatchFiles/Linux/GenerateProjectFiles.sh
@@ -21,17 +21,15 @@ if [ ! -d "$BASE_PATH/../../../Source" ]; then
 fi
 
 source "$BASE_PATH/SetupEnvironment.sh" -dotnet "$BASE_PATH"
-
-# If this is a source drop of the engine make sure that the UnrealBuildTool is up-to-date
-if [ ! -f "$BASE_PATH/../../InstalledBuild.txt" ]; then
-	if [ -f "$BASE_PATH/../../../Source/Programs/UnrealBuildTool/UnrealBuildTool.csproj" ]; then
-		dotnet build "$BASE_PATH/../../../Source/Programs/UnrealBuildTool/UnrealBuildTool.csproj" -c Development -v quiet
-		if [ $? -ne 0 ]; then
-			echo GenerateProjectFiles ERROR: Failed to build UnrealBuildTool
-			exit 1
-		fi
+# ensure UnrealBuildTool is up to date if the project file exists, but not if running from an installed build
+if [ -f "$BASE_PATH/../../../Source/Programs/UnrealBuildTool/UnrealBuildTool.csproj" -a ! -f "$BASE_PATH/../../../Build/InstalledBuild.txt" ]; then
+	"$BASE_PATH/../BuildUBT.sh"
+	if [ $? -ne 0 ]; then
+		echo GenerateProjectFiles ERROR: Failed to build UnrealBuildTool
+		exit 1
 	fi
 fi
 
+
 # pass all parameters to UBT
 dotnet "$BASE_PATH/../../../Binaries/DotNET/UnrealBuildTool/UnrealBuildTool.dll" -projectfiles "$@"
diff --git a/Engine/Build/BatchFiles/Mac/GenerateProjectFiles.sh b/Engine/Build/BatchFiles/Mac/GenerateProjectFiles.sh
--- a/Engine/Build/BatchFiles/Mac/GenerateProjectFiles.sh
+++ b/Engine/Build/BatchFiles/Mac/GenerateProjectFiles.sh
@@ -1,5 +1,6 @@
 #!/bin/bash
 # Copyright Epic Games, Inc. All Rights Reserved.
+
 echo
 echo Setting up Unreal Engine project files...
 echo
@@ -14,7 +15,6 @@ if [ ! -d "$BASE_PATH/../../../Binaries/DotNET" ]; then
 	exit 1
 fi
 
-echo "$BASE_PATH"
 if [ ! -d "$BASE_PATH/../../../Source" ]; then
 	echo GenerateProjectFiles ERROR: This script file does not appear to be located inside the Engine/Build/BatchFiles/Mac directory.
 	exit 1
@@ -23,13 +23,13 @@ fi
 source "$BASE_PATH/SetupEnvironment.sh" -dotnet "$BASE_PATH"
 # ensure UnrealBuildTool is up to date if the project file exists, but not if running from an installed build
 if [ -f "$BASE_PATH/../../../Source/Programs/UnrealBuildTool/UnrealBuildTool.csproj" -a ! -f "$BASE_PATH/../../../Build/InstalledBuild.txt" ]; then
-	dotnet build "$BASE_PATH/../../../Source/Programs/UnrealBuildTool/UnrealBuildTool.csproj" -c Development -v quiet
-
+	"$BASE_PATH/../BuildUBT.sh"
 	if [ $? -ne 0 ]; then
 		echo GenerateProjectFiles ERROR: Failed to build UnrealBuildTool
 		exit 1
 	fi
 fi
 
+
 # pass all parameters to UBT
 dotnet "$BASE_PATH/../../../Binaries/DotNET/UnrealBuildTool/UnrealBuildTool.dll" -projectfiles "$@"
diff --git a/Engine/Build/BatchFiles/RunUAT.bat b/Engine/Build/BatchFiles/RunUAT.bat
--- a/Engine/Build/BatchFiles/RunUAT.bat
+++ b/Engine/Build/BatchFiles/RunUAT.bat
@@ -69,13 +69,6 @@ rem ## check if the UAT projects are present. if not, we'll just use the precomp
 if not exist Source\Programs\AutomationTool\AutomationTool.csproj goto RunPrecompiled
 if not exist Source\Programs\AutomationToolLauncher\AutomationToolLauncher.csproj goto RunPrecompiled
 
-
-rem Checking for out-of-date files won't find source files in places like Engine/Platform, resulting in occasionally
-rem out of date builds. Until a better solution is found, always try to build AutomationTool.
-set FORCECOMPILE_UAT=FORCE
-
-
-
 call "%SCRIPT_DIR%BuildUAT.bat" %MSBUILD_LOGLEVEL% %FORCECOMPILE_UAT%
 if errorlevel 1 goto Error_UATCompileFailed
 
diff --git a/Engine/Build/BatchFiles/RunUAT.sh b/Engine/Build/BatchFiles/RunUAT.sh
--- a/Engine/Build/BatchFiles/RunUAT.sh
+++ b/Engine/Build/BatchFiles/RunUAT.sh
@@ -91,10 +91,6 @@ if [ "$(uname)" = "Linux" ]; then
 	source "$SCRIPT_DIR/Linux/SetupEnvironment.sh" $EnvironmentType "$SCRIPT_DIR/Linux"
 fi
 
-# Checking for out-of-date files won't find source files in places like Engine/Platform, resulting in occasionally
-# out of date builds. Until a better solution is found, always try to build AutomationTool.
-FORCECOMPILE_UAT=FORCE
-
 if [ "$UATCompileArg" = "-compile" ]; then
   # see if the .csproj exists to be compiled
 	if [ ! -f Source/Programs/AutomationTool/AutomationTool.csproj ]; then
diff --git a/Engine/Build/BatchFiles/RunUBT.bat b/Engine/Build/BatchFiles/RunUBT.bat
--- a/Engine/Build/BatchFiles/RunUBT.bat
+++ b/Engine/Build/BatchFiles/RunUBT.bat
@@ -20,31 +20,8 @@ call "%~dp0GetDotnetPath.bat"
 if errorlevel 1 goto Error_NoDotnetSDK
 REM ## Skip msbuild detection if using dotnet as this is done for us by dotnet-cli
 
-rem Check to see if the files in the UBT directory have changed. We conditionally include platform files from the .csproj file, but MSBuild doesn't recognize the dependency when new files are added. 
-md ..\Intermediate\Build >nul 2>nul
-dir /s /b Programs\UnrealBuildTool\*.cs >..\Intermediate\Build\UnrealBuildToolFiles.txt
-
-if not exist ..\Platforms goto NoPlatforms
-for /d %%D in (..\Platforms\*) do (
-	if exist %%D\Source\Programs\UnrealBuildTool dir /s /b %%D\Source\Programs\UnrealBuildTool\*.cs 2> nul >> ..\Intermediate\Build\UnrealBuildToolFiles.txt
-)
-:NoPlatforms
-
-if not exist ..\Restricted goto NoRestricted
-for /d %%D in (..\Restricted\*) do (
-	if exist %%D\Source\Programs\UnrealBuildTool dir /s /b %%D\Source\Programs\UnrealBuildTool\*.cs 2> nul >> ..\Intermediate\Build\UnrealBuildToolFiles.txt
-)
-:NoRestricted
-
-fc /b ..\Intermediate\Build\UnrealBuildToolFiles.txt ..\Intermediate\Build\UnrealBuildToolPrevFiles.txt >nul 2>nul
-if not errorlevel 1 goto SkipClean
-copy /y ..\Intermediate\Build\UnrealBuildToolFiles.txt ..\Intermediate\Build\UnrealBuildToolPrevFiles.txt >nul
-
-dotnet clean Programs\UnrealBuildTool\UnrealBuildTool.csproj -c Development -v quiet
-
-:SkipClean
-echo Building UnrealBuildTool...
-dotnet build Programs\UnrealBuildTool\UnrealBuildTool.csproj -c Development -v quiet
+rem ## Build UnrealBuildTool if necessary
+call "%~dp0BuildUBT.bat"
 if errorlevel 1 goto Error_UBTCompileFailed
 
 rem ## Run UnrealBuildTool to generate Visual Studio solution and project files
diff --git a/Engine/Build/BatchFiles/RunUBT.sh b/Engine/Build/BatchFiles/RunUBT.sh
--- a/Engine/Build/BatchFiles/RunUBT.sh
+++ b/Engine/Build/BatchFiles/RunUBT.sh
@@ -29,7 +29,7 @@ fi
 if [[ "$*" != *-SkipUBTBuild* ]]; then
   if [ -f "$BASE_PATH/../../Source/Programs/UnrealBuildTool/UnrealBuildTool.csproj" ]; then 
     echo "Building UBT..."
-    dotnet build $BASE_PATH/../../Source/Programs/UnrealBuildTool/UnrealBuildTool.csproj -c Development -v quiet
+    "$BASE_PATH/BuildUBT.sh"
 
     if [ $? -ne 0 ]; then
       echo RunUBT ERROR: Failed to build UnrealBuildTool
diff --git a/Engine/Source/Programs/UnrealBuildTool/UnrealBuildTool.csproj b/Engine/Source/Programs/UnrealBuildTool/UnrealBuildTool.csproj
--- a/Engine/Source/Programs/UnrealBuildTool/UnrealBuildTool.csproj
+++ b/Engine/Source/Programs/UnrealBuildTool/UnrealBuildTool.csproj
@@ -102,8 +102,7 @@
   </ItemGroup>
 
   <ItemGroup>
-    <!-- If you add a project reference, BuildUAT.bat/.sh need to be updated to also test for those projects.
-    While similar checks are in RunUBT.bat/.sh and GenerateProjectFiles.bat, they don't currently check for these projects. -->
+    <!-- If you add a project reference, BuildUAT.bat/.sh and BuildUBT.bat/.sh need to be updated to also test for those projects.-->
     <ProjectReference Include="..\Shared\EpicGames.Build\EpicGames.Build.csproj" />
     <ProjectReference Include="..\Shared\EpicGames.Core\EpicGames.Core.csproj" />
     <ProjectReference Include="..\Shared\EpicGames.IoHash\EpicGames.IoHash.csproj" />
@@ -133,4 +132,12 @@
          Adding a loader context resolves many of the issues, but it still fails when MSBuild tries to write the json -->
     <PackageReference Include="Newtonsoft.Json" Version="13.0.1" />
   </ItemGroup>
+
+  <ItemGroup>
+    <None Include="..\..\..\Build\BatchFiles\BuildUBT.bat"><Link>BatchFiles\BuildUBT.bat</Link></None>
+    <None Include="..\..\..\Build\BatchFiles\BuildUBT.sh"><Link>BatchFiles\BuildUBT.sh</Link></None>
+    <None Include="..\..\..\Build\BatchFiles\GenerateProjectFiles.bat"><Link>BatchFiles\GenerateProjectFiles.bat</Link></None>
+    <None Include="..\..\..\Build\BatchFiles\RunUBT.bat"><Link>BatchFiles\RunUBT.bat</Link></None>
+    <None Include="..\..\..\Build\BatchFiles\RunUBT.sh"><Link>BatchFiles\RunUBT.sh</Link></None>
+  </ItemGroup>
 </Project>
