BatchFiles: Remove dotnet build -o argument

#rnx

diff --git a/Engine/Build/BatchFiles/BuildUAT.bat b/Engine/Build/BatchFiles/BuildUAT.bat
--- a/Engine/Build/BatchFiles/BuildUAT.bat
+++ b/Engine/Build/BatchFiles/BuildUAT.bat
@@ -89,8 +89,7 @@ call "%~dp0GetDotnetPath.bat"
 if errorlevel 1 goto Error_NoDotnetSDK
 
 echo Building AutomationTool...
-dotnet clean Programs\AutomationTool\AutomationTool.csproj -c Development -v %MSBUILD_LOGLEVEL% -o ..\Binaries\DotNET\AutomationTool --nologo
-dotnet build Programs\AutomationTool\AutomationTool.csproj -c Development -v %MSBUILD_LOGLEVEL% -o ..\Binaries\DotNET\AutomationTool
+dotnet build Programs\AutomationTool\AutomationTool.csproj -c Development -v %MSBUILD_LOGLEVEL%
 if errorlevel 1 goto Error_UATCompileFailed
 
 rem record input files - regardless of how we got here, these are now our point of reference
diff --git a/Engine/Build/BatchFiles/BuildUBT.bat b/Engine/Build/BatchFiles/BuildUBT.bat
--- a/Engine/Build/BatchFiles/BuildUBT.bat
+++ b/Engine/Build/BatchFiles/BuildUBT.bat
@@ -82,10 +82,8 @@ rem ## Verify that dotnet is present
 call "%~dp0GetDotnetPath.bat"
 if errorlevel 1 goto Error_NoDotnetSDK
 
-:SkipClean
 echo Building UnrealBuildTool...
-dotnet clean Programs\UnrealBuildTool\UnrealBuildTool.csproj -c Development -v %MSBUILD_LOGLEVEL% -o ..\Binaries\DotNET\UnrealBuildTool --nologo
-dotnet build Programs\UnrealBuildTool\UnrealBuildTool.csproj -c Development -v %MSBUILD_LOGLEVEL% -o ..\Binaries\DotNET\UnrealBuildTool
+dotnet build Programs\UnrealBuildTool\UnrealBuildTool.csproj -c Development -v %MSBUILD_LOGLEVEL%
 if errorlevel 1 goto Error_UBTCompileFailed
 
 rem record input files - regardless of how we got here, these are now our point of reference
