Update scripts for 15.0.1 LLVM

diff --git a/Engine/Build/BatchFiles/Linux/Toolchain/BuildLibCxx.sh b/Engine/Build/BatchFiles/Linux/Toolchain/BuildLibCxx.sh
--- a/Engine/Build/BatchFiles/Linux/Toolchain/BuildLibCxx.sh
+++ b/Engine/Build/BatchFiles/Linux/Toolchain/BuildLibCxx.sh
@@ -53,13 +53,19 @@ BuildLibCxx()
       -DCMAKE_TOOLCHAIN_FILE="/tmp/__cmake_toolchain.cmake" \
       -DCMAKE_MAKE_PROGRAM=$(which make) \
       -DPYTHON_EXECUTABLE=$(which python) \
+      -DPython3_EXECUTABLE=$(which python3) \
+      -DLLVM_INCLUDE_BENCHMARKS=OFF \
+      -DLIBCXX_INCLUDE_BENCHMARKS=OFF \
       -DCMAKE_INSTALL_PREFIX=${INSTALL_DIR} \
       -DLLVM_TEMPORARILY_ALLOW_OLD_TOOLCHAIN=TRUE \
       -DCMAKE_BUILD_TYPE=Release \
+      -DLLVM_ENABLE_PER_TARGET_RUNTIME_DIR=OFF \
       -DLLVM_ENABLE_PROJECTS="libcxx;libcxxabi" \
       -DLIBCXX_ENABLE_SHARED=OFF \
-      -DLIBCXXABI_ENABLE_SHARED=ON \
+      -DLIBCXXABI_ENABLE_SHARED=OFF \
+      -DLIBCXX_ENABLE_STATIC=ON \
       -DLIBCXX_ENABLE_STATIC_ABI_LIBRARY=ON \
+      -DLIBCXXABI_INSTALL_SHARED_LIBRARY=OFF \
       -DLIBCXX_CXX_ABI_LIBRARY=c++abi \
       -DLIBCXX_ENABLE_FILESYSTEM=ON \
       ${LLVM_DIR}/llvm
@@ -69,8 +75,8 @@ BuildLibCxx()
     make install-cxx-headers
 
     mkdir -p ${INSTALL_DIR}/lib/Linux/${ARCH}
-    cp lib/libc++.a ${INSTALL_DIR}/lib/Linux/${ARCH}
-    cp lib/libc++abi.a ${INSTALL_DIR}/lib/Linux/${ARCH}
+    cp -v lib/libc++.a ${INSTALL_DIR}/lib/Linux/${ARCH}
+    cp -v lib/libc++abi.a ${INSTALL_DIR}/lib/Linux/${ARCH}
 
     popd
 }
diff --git a/Engine/Build/BatchFiles/Linux/Toolchain/RunMe.bat b/Engine/Build/BatchFiles/Linux/Toolchain/RunMe.bat
--- a/Engine/Build/BatchFiles/Linux/Toolchain/RunMe.bat
+++ b/Engine/Build/BatchFiles/Linux/Toolchain/RunMe.bat
@@ -1,7 +1,7 @@
 @REM @echo off
 
-set TOOLCHAIN_VERSION=v20
-set LLVM_VERSION=13.0.1
+set TOOLCHAIN_VERSION=v21
+set LLVM_VERSION=15.0.1
 
 set SVN_BINARY=%CD%\..\..\..\..\Binaries\ThirdParty\svn\Win64\svn.exe
 set CMAKE_BINARY=%CD%\..\..\..\..\Extras\ThirdPartyNotUE\CMake\bin\cmake.exe
@@ -31,7 +31,7 @@ pushd %TEMP%\clang-build-%LLVM_VERSION%
 
 unzip -o %ROOT_DIR:\=/%/%FILENAME%-windows.zip -d OUTPUT
 
-set GIT_LLVM_RELEASE_HASH=43ff75f2c3feef64f9d73328230d34dac8832a91
+set GIT_LLVM_RELEASE_HASH=b73d2c8c720a8c8e6e73b11be4e27afa6cb75bdf
 
 git clone [REDACTED_PROFILE_URL]/llvm-project source
 pushd source
diff --git a/Engine/Build/BatchFiles/Linux/Toolchain/RunMe.sh b/Engine/Build/BatchFiles/Linux/Toolchain/RunMe.sh
--- a/Engine/Build/BatchFiles/Linux/Toolchain/RunMe.sh
+++ b/Engine/Build/BatchFiles/Linux/Toolchain/RunMe.sh
@@ -6,7 +6,9 @@ SCRIPT_DIR=$(cd "$(dirname "$BASH_SOURCE")" ; pwd)
 SCRIPT_NAME=$(basename "$BASH_SOURCE")
 
 # https://stackoverflow.com/questions/23513045/how-to-check-if-a-process-is-running-inside-docker-container
-if ! $(grep -q "/docker/" /proc/1/cgroup); then
+# https://unix.stackexchange.com/questions/607695/how-to-check-if-its-docker-or-host-machine-inside-bash-script
+# cgroups 2 busted the first way, attempting a new way here but may break again in the future
+if ! [ -f "/.dockerenv" ]; then
 
   ##############################################################################
   # host commands
diff --git a/Engine/Build/BatchFiles/Linux/Toolchain/build_linux_toolchain.sh b/Engine/Build/BatchFiles/Linux/Toolchain/build_linux_toolchain.sh
--- a/Engine/Build/BatchFiles/Linux/Toolchain/build_linux_toolchain.sh
+++ b/Engine/Build/BatchFiles/Linux/Toolchain/build_linux_toolchain.sh
@@ -3,8 +3,8 @@
 set -x
 set -eu
 
-ToolChainVersion=v20
-LLVM_VERSION=13.0.1
+ToolChainVersion=v21
+LLVM_VERSION=15.0.1
 LLVM_URL=[REDACTED_PROFILE_URL]/llvm-project/releases/download/llvmorg-${LLVM_VERSION}
 ZLIB_PATH=/src/v1.2.8
 
@@ -24,31 +24,34 @@ CORES=$(getconf _NPROCESSORS_ONLN)
 echo Using $CORES cores for building
 
 echo "check_certificate=off" > "$HOME/.wgetrc"
-# Get crosstool-ng
-git clone [REDACTED_PROFILE_URL]/crosstool-ng -b 1.22
 
-# Build crosstool-ng
-pushd crosstool-ng
-./bootstrap && ./configure --enable-local && make
-popd
+if [ ! -d "/src/build/crosstool-ng" ]; then
+	# Get crosstool-ng
+	git clone [REDACTED_PROFILE_URL]/crosstool-ng -b 1.22
 
-# Build linux toolchain to OUTPUT-linux
-for arch in $TARGETS; do
-	mkdir -p build-linux-$arch
-	pushd build-linux-$arch
-	cp /src/$arch.linux.config .config
-	../crosstool-ng/ct-ng build.$CORES
+	# Build crosstool-ng
+	pushd crosstool-ng
+	./bootstrap && ./configure --enable-local && make
 	popd
-done
 
-# Build windows toolchain to OUTPUT-windows
-for arch in $TARGETS; do
-	mkdir -p build-windows-$arch
-	pushd build-windows-$arch
-	cp /src/$arch.windows.config .config
-	../crosstool-ng/ct-ng build.$CORES
-	popd
-done
+	# Build linux toolchain to OUTPUT-linux
+	for arch in $TARGETS; do
+		mkdir -p build-linux-$arch
+		pushd build-linux-$arch
+		cp /src/$arch.linux.config .config
+		../crosstool-ng/ct-ng build.$CORES
+		popd
+	done
+
+	# Build windows toolchain to OUTPUT-windows
+	for arch in $TARGETS; do
+		mkdir -p build-windows-$arch
+		pushd build-windows-$arch
+		cp /src/$arch.windows.config .config
+		../crosstool-ng/ct-ng build.$CORES
+		popd
+	done
+fi
 
 # since we are -u in the bash script and this ENV is not set it complains when source the devtoolset-7
 export MANPATH=""
@@ -70,12 +73,14 @@ CLANG=clang-${LLVM_VERSION}
 UNWIND=libunwind-${LLVM_VERSION}
 LLD=lld-${LLVM_VERSION}
 COMPILER_RT=compiler-rt-${LLVM_VERSION}
+CMAKE_MOD=cmake-${LLVM_VERSION}
 
 wget ${LLVM_URL}/${LLVM}.src.tar.xz
 wget ${LLVM_URL}/${CLANG}.src.tar.xz
 wget ${LLVM_URL}/${UNWIND}.src.tar.xz
 wget ${LLVM_URL}/${LLD}.src.tar.xz
 wget ${LLVM_URL}/${COMPILER_RT}.src.tar.xz
+wget ${LLVM_URL}/${CMAKE_MOD}.src.tar.xz
 
 mkdir -p llvm
 tar -xf $LLVM.src.tar.xz --strip-components 1 -C llvm
@@ -88,15 +93,23 @@ tar -xf $LLD.src.tar.xz --strip-components 1 -C llvm/tools/lld/
 mkdir -p llvm/projects/compiler-rt
 tar -xf $COMPILER_RT.src.tar.xz --strip-components 1 -C llvm/projects/compiler-rt/
 
+#llvm has cmake/modules, where compiler-rt has cmake/Modules and cmake.sr.tar.xz has cmake/Modules...
+tar -xf $CMAKE_MOD.src.tar.xz --strip-components 2 -C llvm/cmake/modules/
+tar -xf $CMAKE_MOD.src.tar.xz --strip-components 1 -C llvm/projects/compiler-rt/cmake/
+
 # this fixes an issue where AT_HWCAP2 is just not defined correctly in our sysroot. This is likely due to
 # AT_HWCAP2 being around since glibc 2.18 offically, while we are still stuck on 2.17 glibc.
-patch -d llvm/projects/compiler-rt/ -p 1 < /src/patches/compiler-rt/manually-define-AT_HWCAP2.diff
+patch -d llvm/projects/compiler-rt/ -p 2 < /src/patches/compiler-rt/manually-define-AT_HWCAP2.diff
+
+# move back to defaulting to dwarf 4, as if we leave to dwarf 5 libs built with dwarf5 will force everything to dwarf5
+# even if you request dwarf 4. dwarf 5 currently causes issues with dump_syms and gdb/lldb earlier versions
+patch -d llvm/tools/ -p 1 < /src/patches/clang/default-dwarf-4.patch
 
 # LLVM has just failed to support stand-alone LLD build, cheat by moving a required header into a location it can be found easily
 # if you fulling include this you end up breaking other things. [REDACTED_PROFILE_URL]/llvm-project/issues/48572
 cp -rf llvm/libunwind/include/mach-o/ llvm/include
 
-mkdir build-clang
+mkdir -p build-clang
 pushd build-clang
 	# CMake Error at cmake/modules/CheckCompilerVersion.cmake:40 (message):
 	#   Host GCC version should be at least 5.1 because LLVM will soon use new C++
@@ -114,6 +127,7 @@ pushd build-clang
 		-DLLVM_ENABLE_LIBCXX=1 \
 		-DLLVM_TEMPORARILY_ALLOW_OLD_TOOLCHAIN=ON \
 		-DCMAKE_INSTALL_PREFIX=${InstallClangDir} \
+		-DLLVM_INCLUDE_BENCHMARKS=OFF \
 		-DLLVM_TARGETS_TO_BUILD="AArch64;X86"
 
 	make -j$CORES && make install
@@ -126,21 +140,23 @@ for arch in $TARGETS; do
 	pushd ${OutputDirLinux}/$arch/
 		chmod -R +w .
 
-		# copy $arch/include/c++ to include/c++
-		cp -r -L $arch/include .
+		if [ -d "$arch" ]; then
+			# copy $arch/include/c++ to include/c++
+			cp -r -L $arch/include .
 
-		# copy usr lib64 and include dirs
-		mkdir -p usr
-		cp -r -L $arch/sysroot/usr/include usr
-		cp -r -L $arch/sysroot/usr/lib64 usr
-		cp -r -L $arch/sysroot/usr/lib usr
+			# copy usr lib64 and include dirs
+			mkdir -p usr
+			cp -r -L $arch/sysroot/usr/include usr
+			cp -r -L $arch/sysroot/usr/lib64 usr
+			cp -r -L $arch/sysroot/usr/lib usr
 
-		cp -r -L $arch/lib64 .
-		cp -r -L $arch/lib .
+			cp -r -L $arch/lib64 .
+			cp -r -L $arch/lib .
 
-		[[ -f build.log.bz2 ]] && mv build.log.bz2 ../../build-linux-$arch.log.bz2
+			[[ -f build.log.bz2 ]] && mv build.log.bz2 ../../build-linux-$arch.log.bz2
 
-		rm -rf $arch
+			rm -rf $arch
+		fi
 	popd
 
 	echo "Copying clang..."
@@ -171,7 +187,7 @@ for arch in $TARGETS; do
 	cp -r ${OutputDirLinux}/x86_64-unknown-linux-gnu/lib/clang/${LLVM_VERSION}/share/* ${OutputDirLinux}/$arch/lib/clang/${LLVM_VERSION}/share/
 	cp -r ${OutputDirLinux}/x86_64-unknown-linux-gnu/lib/clang/${LLVM_VERSION}/include/* ${OutputDirLinux}/$arch/lib/clang/${LLVM_VERSION}/include/
 
-	mkdir build-rt-$arch
+	mkdir -p build-rt-$arch
 	pushd build-rt-$arch
 
 		cmake3 -G "Unix Makefiles" ../llvm/projects/compiler-rt \
@@ -191,6 +207,8 @@ for arch in $TARGETS; do
 			-DCMAKE_C_FLAGS="--target=$arch --sysroot=${OutputDirLinux}/$arch" \
 			-DCMAKE_CXX_FLAGS="--target=$arch --sysroot=${OutputDirLinux}/$arch" \
 			-DCMAKE_ASM_FLAGS="--target=$arch --sysroot=${OutputDirLinux}/$arch" \
+			-DCOMPILER_RT_BUILD_ORC=OFF \
+			-DCOMPILER_RT_BUILD_LIBFUZZER=OFF \
 			-DCMAKE_INSTALL_PREFIX=../install-rt-$arch \
 			-DSANITIZER_COMMON_LINK_FLAGS="-fuse-ld=lld" \
 			-DSCUDO_LINK_FLAGS="-fuse-ld=lld" \
