AutomationTool: Add experimental support for running UAT under Wine

Based on work by [REDACTED_PERSON].

diff --git a/Engine/Build/BatchFiles/RunUAT.bat b/Engine/Build/BatchFiles/RunUAT.bat
--- a/Engine/Build/BatchFiles/RunUAT.bat
+++ b/Engine/Build/BatchFiles/RunUAT.bat
@@ -9,11 +9,11 @@ rem ## if you copy it to a different location and run it.
 setlocal EnableExtensions
 echo Running AutomationTool...
 
-rem uppercase the drive letter
-SETLOCAL ENABLEDELAYEDEXPANSION
+rem Uppercase the drive letter
 set DRIVE_LETTER=%~d0
-FOR %%Z IN (A B C D E F G H I J K L M N O P Q R S T U V W X Y Z) DO SET DRIVE_LETTER=!DRIVE_LETTER:%%Z=%%Z!
-SETLOCAL DISABLEDELAYEDEXPANSION
+FOR %%Z IN (A B C D E F G H I J K L M N O P Q R S T U V W X Y Z) DO (
+	IF /I %DRIVE_LETTER%==%%Z: SET DRIVE_LETTER=%%Z:
+)
 
 set SCRIPT_DIR=%DRIVE_LETTER%%~p0
 set UATExecutable=AutomationTool.dll
diff --git a/Engine/Build/BatchFiles/RunWineUAT.sh b/Engine/Build/BatchFiles/RunWineUAT.sh
new file mode 100644
--- /dev/null
+++ b/Engine/Build/BatchFiles/RunWineUAT.sh
@@ -0,0 +1,75 @@
+#!/bin/bash
+# Copyright Epic Games, Inc. All Rights Reserved.
+# Entrypoint for UAT when running under Linux and Wine
+# Building under Wine is highly experimental and not officially supported
+
+set -o errexit
+set -o pipefail
+
+# Put ourselves into Engine directory (two up from location of this script)
+SCRIPT_DIR=$(cd -- "$(dirname -- "${BASH_SOURCE[0]}")" &> /dev/null && pwd)
+cd "$SCRIPT_DIR/../.." || exit
+
+export WINEPREFIX="$HOME/.wine"
+export WINEDEBUG=-all
+export WINEARCH=win64
+export DISPLAY=:99
+
+# Remove empty "qagame" directory. This collides with directory "QAGame" and confuses directory caching in UAT/UBT.
+# This is a problem on Linux with ext4 and the Perforce history for this directory.
+rm -rf ../qagame
+
+[ ! -d "$UE_SDKS_ROOT/HostWin64" ] && echo "AutoSDK dir '$UE_SDKS_ROOT' does not contain sub-dir 'HostWin64'" && exit 1
+# Reformat UE_SDKS_ROOT under Z: and with backslashes
+export UE_SDKS_ROOT="Z:${UE_SDKS_ROOT//\//\\}"
+
+# Pretty hacky, depends on this package being around but wine/windows works so much better with an x11 server running
+Xvfb :99 -screen 0 1024x768x16 &
+XVFB_PID=$!
+if ! (ps -p $XVFB_PID > /dev/null); then echo 'Failed to start xvfb process'; exit 1; fi
+
+function on_exit {
+    echo "Kill XVFB '$XVFB_PID'"
+    kill $XVFB_PID
+}
+trap on_exit EXIT
+
+# Make WineUAT are built on Linux *not* Wine/Windows as it is very slow and broken
+# if UBT is not built, building WineUAT will build it
+
+# See if we have the no compile arg
+if echo "$@" | grep -q -w -i "\-nocompile"; then
+    WineUATCompileArg=
+else
+    WineUATCompileArg=-compile
+fi
+
+# Control toggling of msbuild verbosity for easier debugging
+if echo "$@" | grep -q -w -i "\-msbuild-verbose"; then
+    MSBuild_Verbosity=normal
+else
+    MSBuild_Verbosity=quiet
+fi
+
+if [ "$WineUATCompileArg" = "-compile" ]; then
+    # See if the .csproj exists to be compiled
+    if [ ! -f Source/Programs/AutomationTool/AutomationTool.csproj ]; then
+        echo No project to compile, attempting to use precompiled AutomationTool
+        WineUATCompileArg=
+    else
+        "$SCRIPT_DIR/BuildUAT.sh" $MSBuild_Verbosity
+        if [ $? -ne 0 ]; then
+            echo RunWineUAT ERROR: AutomationTool failed to compile.
+            exit 1
+        fi
+
+        # Run UAT once on Linux to initialize and compile scripts
+        # Currently, .NET builds don't compile under Wine
+        #
+        # Eat normal output, but still show errors
+        "$SCRIPT_DIR"/RunUAT.sh -help > /dev/null
+    fi
+fi
+
+# Never UBT/UAT compile through Wine (see comment above)
+wine64 "$SCRIPT_DIR/RunUAT.bat" -nocompileuat -nocompile -NoP4 "$@" 2>&1
\ No newline at end of file
