AutomationTool: Avoid xvfb printing error message when running as non-root under Wine

Skips listening to TCP and Unix sockets.

diff --git a/Engine/Build/BatchFiles/RunWineUAT.sh b/Engine/Build/BatchFiles/RunWineUAT.sh
--- a/Engine/Build/BatchFiles/RunWineUAT.sh
+++ b/Engine/Build/BatchFiles/RunWineUAT.sh
@@ -24,7 +24,8 @@ rm -rf ../qagame
 export UE_SDKS_ROOT="Z:${UE_SDKS_ROOT//\//\\}"
 
 # Pretty hacky, depends on this package being around but wine/windows works so much better with an x11 server running
-Xvfb :99 -screen 0 1024x768x16 &
+# Avoids listening on TCP and Unix sockets as that yields errors when running as non-root
+Xvfb :99 -screen 0 1024x768x16 -nolisten tcp -nolisten unix &
 XVFB_PID=$!
 if ! (ps -p $XVFB_PID > /dev/null); then echo 'Failed to start xvfb process'; exit 1; fi
 
