- Updating InstallP4VUtils.bat to match the changes in mac/linux script

diff --git a/Engine/Build/BatchFiles/InstallP4VUtils.bat b/Engine/Build/BatchFiles/InstallP4VUtils.bat
--- a/Engine/Build/BatchFiles/InstallP4VUtils.bat
+++ b/Engine/Build/BatchFiles/InstallP4VUtils.bat
@@ -1,13 +1,23 @@
+@rem Copyright Epic Games, Inc. All Rights Reserved.
+
+@echo off
 setlocal
 
-set LIBPATH=%LOCALAPPDATA%\Unreal Engine\P4VUtils
+cd %~dp0%
+
+set DESTDIR=%LOCALAPPDATA%\Unreal Engine\P4VUtils
+set SOURCEDIR=..\..\..\Extras\P4VUtils\Binaries\Win64
+set MESSAGE=Installing P4VUtils into p4v...
 
-echo Copying P4VUtils files to %LIBPATH%...
-mkdir "%LIBPATH%" 2> NUL
+dir
+dir ..\..\Restricted\NotForLicensees\Extras\P4VUtils
+if EXIST ..\..\Restricted\NotForLicensees\Extras\P4VUtils (
+	set MESSAGE=Installing P4VUtils [with Epic extensions] into p4v...
+	set SOURCEDIR=..\..\Restricted\NotForLicensees\Extras\P4VUtils\Binaries\Win64
+)
 
-copy ..\..\Extras\P4VUtils\Binaries\Win64\* "%LIBPATH%"
-copy ..\..\Extras\P4VUtils\P4VUtils.ini "%LIBPATH%"
-xcopy /s /y ..\..\Restricted\NotForLicensees\Extras\P4VUtils\* "%LIBPATH%\NotForLicensees\" 2> NUL
+rmdir /s /q "%DESTDIR%" 2> NUL
+xcopy /s /y /i "%SOURCEDIR%" "%DESTDIR%" 2> NUL
 
-echo Installing P4VUtils into p4v...
-dotnet "%LIBPATH%\P4VUtils.dll" install
+echo %MESSAGE%
+"%DESTDIR%\P4VUtils" install
