- Replacement script for dsymutil that can take a universal binary and generate a dsym, when normally dsymutil would error out with > 4gb offsets

diff --git a/Engine/Build/BatchFiles/Mac/GenerateUniversalDSYM.sh b/Engine/Build/BatchFiles/Mac/GenerateUniversalDSYM.sh
new file mode 100755
--- /dev/null
+++ b/Engine/Build/BatchFiles/Mac/GenerateUniversalDSYM.sh
@@ -0,0 +1,70 @@
+#!/bin/bash
+# Copyright Epic Games, Inc. All Rights Reserved.
+
+set -e
+
+echo $1 $2
+
+# $1 is the path to the executable
+# $2 is the path to the dSYM we want to create
+
+# make sure extension of param 2 is a dsym since we are about to delete it, protect against pass / or something
+extension=${2##*.}
+if [[ "${extension}" != "dSYM" &&  "${extension}" != "dsym" ]]; then
+	echo "Usage: GenerateUniversalDSYMl.sh <path to input binary> <path to output dsym>"
+	echo "       For executables in a .app, specify the full path (such as "Foo.app/Contents/MacOS/Foo")"
+	exit 0
+fi
+
+# Cleanup any old dSYM
+if [ -e $2 ]; then
+	echo "Deleting existing .dSYM"
+	rm -r "$2"
+fi
+
+mkdir -p "$(dirname $2)"
+binaryname=$(basename "$1")
+tempdir=`mktemp -d`
+
+# See if the binary is fat (arches will be "x86_64 arm64" for universal)
+arches=(`lipo -archs "$1"`)
+if [ ${#arches[@]} -gt 1 ]; then
+	for i in ${!arches[@]} 
+	do
+		arch=${arches[$i]}
+
+		echo "Extracting $arch from $1..."
+		lipo -extract $arch -output $tempdir/$arch "$1" 
+
+		# Get paths to the extracted binary and the binary inside the dSYM
+		dsympaths[$i]="$tempdir/$arch.dSYM"
+		binpaths[$i]="$tempdir/$arch.dSYM/Contents/Resources/DWARF/$arch"
+	done
+
+	for i in ${!arches[@]} 
+	do
+		arch=${arches[$i]}
+
+		echo "Generating debug symbols (dSYM) for $arch" 
+		dsymutil -o $tempdir/$arch.dSYM $tempdir/$arch &    
+	done
+
+	# wait for the dsymutils to finish
+	wait
+
+	# Copy the first dSYM back to result dSYM location - we will overwrite one file inside the dir structure
+	ditto ${dsympaths[0]} "$2"
+
+	# remove the arch-named file inside, will replace with lipo below
+	rm "$2/Contents/Resources/DWARF/${arches[0]}"
+  
+	# lipo all dsym binaries directly to the project-named file that we'd expect
+	echo "Merging architectures '${binpaths[*]}' together into $2"
+	lipo ${binpaths[*]} -create -output "$2/Contents/Resources/DWARF/${binaryname}"
+
+	# Clean up
+#	rm -rf "$tempdir"
+else
+  echo "Using standard dsymutil because the binary was not universal..."
+  dsymutil "$1" -o "$2"
+fi
diff --git a/Engine/Source/Programs/AutomationTool/Scripts/GenerateDSYM.Automation.cs b/Engine/Source/Programs/AutomationTool/Scripts/GenerateDSYM.Automation.cs
--- a/Engine/Source/Programs/AutomationTool/Scripts/GenerateDSYM.Automation.cs
+++ b/Engine/Source/Programs/AutomationTool/Scripts/GenerateDSYM.Automation.cs
@@ -90,7 +90,23 @@ public class GenerateDSYM : BuildCommand
 
 				// put dSYM next to the binary
 				string dSYM = Path.ChangeExtension(Binary, ".dSYM");
-				string Command = string.Format("-c 'rm -rf \"{0}\"; dsymutil -o \"{0}\" {2} \"{1}\"'", dSYM, Binary, bSaveFlat ? "--flat" : "");
+				if (Binary.Contains(".app/"))
+				{
+					// or the .app if the binary is inside one
+					dSYM = Path.Combine(Path.GetDirectoryName(Binary.Substring(0, Binary.LastIndexOf(".app"))), Path.GetFileName(dSYM));
+				}
+
+				Log.TraceInformation("  Generating for {0} -> {1}", Binary, dSYM);
+
+				string Command;
+				if (bSaveFlat)
+				{
+					Command = string.Format("-c 'rm -rf \"{0}\"; dsymutil -o \"{0}\" {2} \"{1}\"'", dSYM, Binary, bSaveFlat ? "--flat" : "");
+				}
+				else 
+				{ 
+					Command = $"-c '{Unreal.EngineDirectory}/Build/BatchFiles/Mac/GenerateUniversalDSYM.sh \"{Binary}\" \"{dSYM}\"'";
+				}	
 
 				DateTime RunStart = DateTime.Now;
 				Run("bash", Command, Options: ERunOptions.NoLoggingOfRunCommand);
@@ -127,4 +143,4 @@ public class GenerateDSYM : BuildCommand
 
 		}
 	}
-}
+}
\ No newline at end of file
diff --git a/Engine/Source/Programs/UnrealBuildTool/Platform/Mac/MacToolChain.cs b/Engine/Source/Programs/UnrealBuildTool/Platform/Mac/MacToolChain.cs
--- a/Engine/Source/Programs/UnrealBuildTool/Platform/Mac/MacToolChain.cs
+++ b/Engine/Source/Programs/UnrealBuildTool/Platform/Mac/MacToolChain.cs
@@ -969,8 +969,8 @@ namespace UnrealBuildTool
 			// Deletes ay existing file on the building machine. Also, waits 30 seconds, if needed, for the input file to be created in an attempt to work around
 			// a problem where dsymutil would exit with an error saying the input file did not exist.
 			// Note that the source and dest are switched from a copy command
-			string ExtraOptions;
-			string DsymutilPath = GetDsymutilPath(Logger, out ExtraOptions, bIsForLTOBuild: false);
+			string ExtraOptions = "";
+			string DsymutilPath = FileReference.Combine(Unreal.EngineDirectory, "Build/BatchFiles/Mac/GenerateUniversalDSYM.sh").FullName;
 
 			string ArgumentString = "-c \"";
 			ArgumentString += string.Format("for i in {{1..30}}; ");
@@ -985,7 +985,7 @@ namespace UnrealBuildTool
 			ArgumentString += string.Format("if [ ! -f \\\"{1}\\\" ] || [ \\\"{0}\\\" -nt \\\"{1}\\\" ] ; ", MachOBinary.AbsolutePath, OutputFile.AbsolutePath);
 			ArgumentString += string.Format("then ");
 				ArgumentString += string.Format("rm -rf \\\"{0}\\\"; ", OutputFile.AbsolutePath);
-				ArgumentString += string.Format(" \\\"{0}\\\" {3} -f \\\"{1}\\\" -o \\\"{2}\\\"; ",
+				ArgumentString += string.Format(" \\\"{0}\\\" \\\"{1}\\\" \\\"{2}\\\"; ",
 					DsymutilPath,
 					MachOBinary.AbsolutePath,
 					OutputFile.AbsolutePath,
