- Removed a comment which was leaving temp files around instead of deleting them

diff --git a/Engine/Build/BatchFiles/Mac/GenerateUniversalDSYM.sh b/Engine/Build/BatchFiles/Mac/GenerateUniversalDSYM.sh
--- a/Engine/Build/BatchFiles/Mac/GenerateUniversalDSYM.sh
+++ b/Engine/Build/BatchFiles/Mac/GenerateUniversalDSYM.sh
@@ -63,7 +63,7 @@ if [ ${#arches[@]} -gt 1 ]; then
 	lipo ${binpaths[*]} -create -output "$2/Contents/Resources/DWARF/${binaryname}"
 
 	# Clean up
-#	rm -rf "$tempdir"
+	rm -rf "$tempdir"
 else
   echo "Using standard dsymutil because the binary was not universal..."
   dsymutil "$1" -o "$2"
