Suppress all messages from FC (file compare) as we only care for the exit code of the command.
Some messages where confusing the customers.

diff --git a/Engine/Build/BatchFiles/BuildUAT.bat b/Engine/Build/BatchFiles/BuildUAT.bat
--- a/Engine/Build/BatchFiles/BuildUAT.bat
+++ b/Engine/Build/BatchFiles/BuildUAT.bat
@@ -85,7 +85,7 @@ if /I "%ARGUMENT%" == "FORCE" goto Build_AutomationTool
 if not exist ..\Binaries\DotNET\AutomationTool\AutomationTool.dll goto Build_AutomationTool
 set RUNUAT_EXITCODE=0
 rem per https://ss64.com/nt/fc.html using redirection syntax rather than errorlevel, based on observed inconsistent results from this function
-fc ..\Intermediate\Build\AutomationToolFiles.txt ..\Intermediate\Build\AutomationToolPrevFiles.txt >nul && goto Exit
+fc ..\Intermediate\Build\AutomationToolFiles.txt ..\Intermediate\Build\AutomationToolPrevFiles.txt >nul 2>&1 && goto Exit
 
 :Build_AutomationTool
 rem ## Verify that dotnet is present
diff --git a/Engine/Build/BatchFiles/BuildUBT.bat b/Engine/Build/BatchFiles/BuildUBT.bat
--- a/Engine/Build/BatchFiles/BuildUBT.bat
+++ b/Engine/Build/BatchFiles/BuildUBT.bat
@@ -79,7 +79,7 @@ if /I "%ARGUMENT%" == "FORCE" goto Build_UnrealBuildTool
 if not exist ..\Binaries\DotNET\UnrealBuildTool\UnrealBuildTool.dll goto Build_UnrealBuildTool
 set RUNUBT_EXITCODE=0
 rem per https://ss64.com/nt/fc.html using redirection syntax rather than errorlevel, based on observed inconsistent results from this function
-fc ..\Intermediate\Build\UnrealBuildToolFiles.txt ..\Intermediate\Build\UnrealBuildToolPrevFiles.txt >nul && goto Exit
+fc ..\Intermediate\Build\UnrealBuildToolFiles.txt ..\Intermediate\Build\UnrealBuildToolPrevFiles.txt >nul 2>&1 && goto Exit
 
 :Build_UnrealBuildTool
 rem ## Verify that dotnet is present
