Suppress BuildU[A|B]T.bat "File Not Found" output when they search for dependencies

#rnx

diff --git a/Engine/Build/BatchFiles/BuildUAT.bat b/Engine/Build/BatchFiles/BuildUAT.bat
--- a/Engine/Build/BatchFiles/BuildUAT.bat
+++ b/Engine/Build/BatchFiles/BuildUAT.bat
@@ -50,7 +50,7 @@ for /d %%D in (..\Platforms\*) do (
 		 %%D\Source\Programs\UnrealBuildTool\*.cs ^
 		 %%D\Source\Programs\UnrealBuildTool\*.csproj ^
 		 | find ".cs" >> ..\Intermediate\Build\AutomationToolFiles.txt
-	)
+	) 2>nul
 )
 :NoPlatforms
 
@@ -63,7 +63,7 @@ for /d %%D in (..\Restricted\*) do (
 		 %%D\Source\Programs\UnrealBuildTool\*.cs ^
 		 %%D\Source\Programs\UnrealBuildTool\*.csproj ^
 		 | find ".cs" >> ..\Intermediate\Build\AutomationToolFiles.txt
-	)
+	) 2>nul
 )
 :NoRestricted
 
diff --git a/Engine/Build/BatchFiles/BuildUBT.bat b/Engine/Build/BatchFiles/BuildUBT.bat
--- a/Engine/Build/BatchFiles/BuildUBT.bat
+++ b/Engine/Build/BatchFiles/BuildUBT.bat
@@ -48,7 +48,7 @@ for /d %%D in (..\Platforms\*) do (
 		 %%D\Source\Programs\UnrealBuildTool\*.cs ^
 		 %%D\Source\Programs\UnrealBuildTool\*.csproj ^
 		 | find ".cs" >> ..\Intermediate\Build\UnrealBuildToolFiles.txt
-	)
+	) 2>nul
 )
 :NoPlatforms
 
@@ -59,7 +59,7 @@ for /d %%D in (..\Restricted\*) do (
 		 %%D\Source\Programs\UnrealBuildTool\*.cs ^
 		 %%D\Source\Programs\UnrealBuildTool\*.csproj ^
 		 | find ".cs" >> ..\Intermediate\Build\UnrealBuildToolFiles.txt
-	)
+	) 2>nul
 )
 :NoRestricted
 
