Update UAT & UBT shell script dependency check

#rnx

diff --git a/Engine/Build/BatchFiles/BuildUAT.bat b/Engine/Build/BatchFiles/BuildUAT.bat
--- a/Engine/Build/BatchFiles/BuildUAT.bat
+++ b/Engine/Build/BatchFiles/BuildUAT.bat
@@ -21,6 +21,8 @@ rem find ".cs" files to only lines that match those names - excludes lines that
 md ..\Intermediate\Build >nul 2>nul
 
 dir /s ^
+ Programs\Shared\EpicGames.Box\*.cs ^
+ Programs\Shared\EpicGames.Box\*.csproj ^
  Programs\Shared\EpicGames.Build\*.cs ^
  Programs\Shared\EpicGames.Build\*.csproj ^
  Programs\Shared\EpicGames.Core\*.cs ^
diff --git a/Engine/Build/BatchFiles/BuildUAT.sh b/Engine/Build/BatchFiles/BuildUAT.sh
--- a/Engine/Build/BatchFiles/BuildUAT.sh
+++ b/Engine/Build/BatchFiles/BuildUAT.sh
@@ -37,6 +37,7 @@ elif [ ! -f ../Binaries/DotNET/AutomationTool/AutomationTool.dll ]; then
 
 elif [ -f ../Intermediate/Build/AutomationToolLastBuildTime ]; then
   UPDATED_DEP_FILES="$(find \
+    Programs/Shared/EpicGames.Box \
     Programs/Shared/EpicGames.Build \
     Programs/Shared/EpicGames.Core \
     Programs/Shared/EpicGames.Horde \
diff --git a/Engine/Build/BatchFiles/BuildUBT.bat b/Engine/Build/BatchFiles/BuildUBT.bat
--- a/Engine/Build/BatchFiles/BuildUBT.bat
+++ b/Engine/Build/BatchFiles/BuildUBT.bat
@@ -21,6 +21,8 @@ rem find ".cs" files to only lines that match those names - excludes lines that
 md ..\Intermediate\Build >nul 2>nul
 
 dir /s ^
+ Programs\Shared\EpicGames.Box\*.cs ^
+ Programs\Shared\EpicGames.Box\*.csproj ^
  Programs\Shared\EpicGames.Build\*.cs ^
  Programs\Shared\EpicGames.Build\*.csproj ^
  Programs\Shared\EpicGames.Core\*.cs ^
diff --git a/Engine/Build/BatchFiles/BuildUBT.sh b/Engine/Build/BatchFiles/BuildUBT.sh
--- a/Engine/Build/BatchFiles/BuildUBT.sh
+++ b/Engine/Build/BatchFiles/BuildUBT.sh
@@ -37,6 +37,7 @@ elif [ ! -f ../Binaries/DotNET/UnrealBuildTool/UnrealBuildTool.dll ]; then
 
 elif [ -f ../Intermediate/Build/UnrealBuildToolLastBuildTime ]; then
   UPDATED_DEP_FILES="$(find \
+    Programs/Shared/EpicGames.Box \
     Programs/Shared/EpicGames.Build \
     Programs/Shared/EpicGames.Core \
     Programs/Shared/EpicGames.Horde \
