Script fixes for v22 toolchain

- Windows binaries use rpmalloc for a tremendous LTO speed improvements
- also a patch from clang 17.x to add a flag to disable auto-upgrading debug info
- new patch for our old sysroot missing modern stuff (HWCAP_CRC32)
- switched Linux script away from clang/lld standalone builds because they are no longer supported and don't work reliably
- fixed Windows build to allow compressing debug info with zlib
- cleaned up Windows batch files so it uses native commands (except for git/patch) and can run without Cygwin/etc installed

diff --git a/Engine/Build/BatchFiles/Linux/Toolchain/RunMe.bat b/Engine/Build/BatchFiles/Linux/Toolchain/RunMe.bat
--- a/Engine/Build/BatchFiles/Linux/Toolchain/RunMe.bat
+++ b/Engine/Build/BatchFiles/Linux/Toolchain/RunMe.bat
@@ -1,69 +1,68 @@
 @REM @echo off
 
-set TOOLCHAIN_VERSION=v21
-set LLVM_VERSION=15.0.1
+set TOOLCHAIN_VERSION=v22
+set LLVM_VERSION=16.0.6
+set LLVM_BRANCH=release/16.x
+set LLVM_TAG=llvmorg-16.0.6
 
-set SVN_BINARY=%CD%\..\..\..\..\Binaries\ThirdParty\svn\Win64\svn.exe
 set CMAKE_BINARY=%CD%\..\..\..\..\Extras\ThirdPartyNotUE\CMake\bin\cmake.exe
 set PYTHON_BINARY=%CD%\..\..\..\..\Binaries\ThirdParty\Python\Win64\python.exe
+set ZLIB_PATH=%CD%\..\..\..\..\..\Engine\Source\ThirdParty\zlib\v1.2.8
+set "ZLIB_PATH=%ZLIB_PATH:\=/%"
 set NSIS_BINARY=C:\Program Files (x86)\NSIS\Bin\makensis.exe
+set PATCH_BINARY=C:\Program Files\Git\usr\bin\patch.exe
+set VS_VERSION="Visual Studio 17 2022"
 
 for %%i in (python.exe) do set PYTHON_BINARY="%%~$PATH:i"
 for %%i in (cmake.exe) do set CMAKE_BINARY="%%~$PATH:i"
-for %%i in (svn.exe) do set SVN_BINARY="%%~$PATH:i"
 
 set FILENAME=%TOOLCHAIN_VERSION%_clang-%LLVM_VERSION%-centos7
 
 echo Building %FILENAME%.exe...
 
 echo.
-echo Using SVN: %SVN_BINARY%
 echo Using CMake: %CMAKE_BINARY%
 echo Using Python: %PYTHON_BINARY%
+echo Using VisualStudio: %VS_VERSION%
 
 @REM We need to build in a directory with shorter path, so we avoid hitting path max limit.
 set ROOT_DIR=%CD%
 
-rm -rf %TEMP:\=/%\clang-build-%LLVM_VERSION%
+rmdir /S /Q %TEMP%\clang-build-%LLVM_VERSION%
 mkdir %TEMP%\clang-build-%LLVM_VERSION%
 pushd %TEMP%\clang-build-%LLVM_VERSION%
 
 
-unzip -o %ROOT_DIR:\=/%/%FILENAME%-windows.zip -d OUTPUT
+unzip -q -o %ROOT_DIR%\%FILENAME%-windows.zip -d OUTPUT
 
-set GIT_LLVM_RELEASE_HASH=b73d2c8c720a8c8e6e73b11be4e27afa6cb75bdf
+echo Cloning rpmalloc to speed up clang/lld runtime operation
+git clone [REDACTED_PROFILE_URL]/rpmalloc rpmalloc --depth 1
 
-git clone [REDACTED_PROFILE_URL]/llvm-project source
+
+echo Cloning LLVM (tag %LLVM_TAG% only)
+rem clone -b can also accept tag names
+git clone [REDACTED_PROFILE_URL]/llvm-project source -b %LLVM_TAG% --single-branch --depth 1 -c advice.detachedHead=false
 pushd source
-git checkout %GIT_LLVM_RELEASE_HASH%
+git -c advice.detachedHead=false checkout tags/%LLVM_TAG% -b %LLVM_BRANCH%
 popd
 
-mkdir build_llvm
-pushd build_llvm
-
-%CMAKE_BINARY% -G "Visual Studio 16 2019" -DLLVM_TEMPORARILY_ALLOW_OLD_TOOLCHAIN=ON -DCMAKE_INSTALL_PREFIX="..\install" -DPYTHON_EXECUTABLE="%PYTHON_BINARY%" "..\source\llvm"
-%CMAKE_BINARY% --build . --target install --config MinSizeRel
-
+echo Applying patches
+pushd "source"
+rem set DRY_RUN=--dry-run
+set DRY_RUN=
+"%PATCH_BINARY%" %DRY_RUN% -p 1 -i %ROOT_DIR%\patches\clang\default-dwarf-4.patch 
+"%PATCH_BINARY%" %DRY_RUN% -p 1 -i %ROOT_DIR%\patches\compiler-rt\manually-define-AT_HWCAP2.diff 
+"%PATCH_BINARY%" %DRY_RUN% -p 1 -i %ROOT_DIR%\patches\llvm\disable-auto-upgrade-debug-info.patch 
 popd
 
-mkdir build_lld
-pushd build_lld
-
-%CMAKE_BINARY% -G "Visual Studio 16 2019" -DLLVM_TEMPORARILY_ALLOW_OLD_TOOLCHAIN=ON -DCMAKE_INSTALL_PREFIX="..\install" -DPYTHON_EXECUTABLE="%PYTHON_BINARY%" -DLLVM_CONFIG_PATH="..\install\bin\llvm-config.exe" "..\source\lld"
-%CMAKE_BINARY% -G "Visual Studio 16 2019" -DLLVM_TEMPORARILY_ALLOW_OLD_TOOLCHAIN=ON -DCMAKE_INSTALL_PREFIX="..\install" -DPYTHON_EXECUTABLE="%PYTHON_BINARY%" -DLLVM_CONFIG_PATH="..\install\bin\llvm-config.exe" "..\source\lld"
-%CMAKE_BINARY% --build . --target install --config MinSizeRel
 
+echo Building LLVM, clang, lld, bolt
+mkdir build_all
+pushd build_all
+%CMAKE_BINARY% -G %VS_VERSION% -DLLVM_ENABLE_PROJECTS=llvm;clang;lld;bolt -DLLVM_INTEGRATED_CRT_ALLOC=../rpmalloc -DLLVM_USE_CRT_RELEASE=MT -DCMAKE_BUILD_TYPE=Release -DCMAKE_C_FLAGS_RELEASE="/O2 /Ob3 /DNDEBUG /Zi /Gy" -DCMAKE_CXX_FLAGS_RELEASE="/O2 /Ob3 /DNDEBUG /Zi /Gy" -DCMAKE_EXE_LINKER_FLAGS_RELEASE="/DEBUG /INCREMENTAL:NO /OPT:REF /OPT:ICF" -DCMAKE_INSTALL_PREFIX="..\install" -DLLVM_INCLUDE_BENCHMARKS=OFF -DLLVM_TARGETS_TO_BUILD="AArch64;X86" -DLLVM_ENABLE_ZLIB=FORCE_ON -DZLIB_LIBRARY="%ZLIB_PATH%/lib/Win64-llvm/Release/zlibstatic.lib" -DZLIB_INCLUDE_DIR="%ZLIB_PATH%/include/Win64/VS2015"  "..\source\llvm"
+%CMAKE_BINARY% --build . --target install --config Release
 popd
 
-mkdir build_clang
-pushd build_clang
-
-%CMAKE_BINARY% -G "Visual Studio 16 2019" -DCMAKE_INSTALL_PREFIX="..\install" -DPYTHON_EXECUTABLE="%PYTHON_BINARY%" "..\source\clang"
-%CMAKE_BINARY% --build . --target install --config MinSizeRel
-
-popd
-
-
 for %%G in (aarch64-unknown-linux-gnueabi x86_64-unknown-linux-gnu) do (
     mkdir OUTPUT\%%G
     mkdir OUTPUT\%%G\bin
@@ -87,8 +86,8 @@ echo %TOOLCHAIN_VERSION%_clang-%LLVM_VERSION%-centos7> OUTPUT\ToolchainVersion.t
 echo Packing final toolchain...
 
 pushd OUTPUT
-rm -rf %ROOT_DIR:\=/%/%FILENAME%.zip
-zip %ROOT_DIR:\=/%/%FILENAME%.zip *
+del /S /Q %ROOT_DIR%\%FILENAME%.zip
+zip -r %ROOT_DIR%\%FILENAME%.zip *
 popd
 
 if exist "%NSIS_BINARY%" (
diff --git a/Engine/Build/BatchFiles/Linux/Toolchain/build_linux_toolchain.sh b/Engine/Build/BatchFiles/Linux/Toolchain/build_linux_toolchain.sh
--- a/Engine/Build/BatchFiles/Linux/Toolchain/build_linux_toolchain.sh
+++ b/Engine/Build/BatchFiles/Linux/Toolchain/build_linux_toolchain.sh
@@ -3,8 +3,12 @@
 set -x
 set -eu
 
-ToolChainVersion=v21
-LLVM_VERSION=15.0.1
+ToolChainVersion=v22
+LLVM_VERSION_MAJOR=16
+LLVM_VERSION=${LLVM_VERSION_MAJOR}.0.6
+LLVM_BRANCH=release/${LLVM_VERSION_MAJOR}.x
+LLVM_TAG=llvmorg-${LLVM_VERSION}
+
 LLVM_URL=[REDACTED_PROFILE_URL]/llvm-project/releases/download/llvmorg-${LLVM_VERSION}
 ZLIB_PATH=/src/v1.2.8
 
@@ -67,47 +71,32 @@ unset LD_LIBRARY_PATH
 # Linux
 #
 
-# Build clang
-LLVM=llvm-${LLVM_VERSION}
-CLANG=clang-${LLVM_VERSION}
-UNWIND=libunwind-${LLVM_VERSION}
-LLD=lld-${LLVM_VERSION}
-COMPILER_RT=compiler-rt-${LLVM_VERSION}
-CMAKE_MOD=cmake-${LLVM_VERSION}
-
-wget ${LLVM_URL}/${LLVM}.src.tar.xz
-wget ${LLVM_URL}/${CLANG}.src.tar.xz
-wget ${LLVM_URL}/${UNWIND}.src.tar.xz
-wget ${LLVM_URL}/${LLD}.src.tar.xz
-wget ${LLVM_URL}/${COMPILER_RT}.src.tar.xz
-wget ${LLVM_URL}/${CMAKE_MOD}.src.tar.xz
-
-mkdir -p llvm
-tar -xf $LLVM.src.tar.xz --strip-components 1 -C llvm
-mkdir -p llvm/tools/clang
-tar -xf $CLANG.src.tar.xz --strip-components 1 -C llvm/tools/clang/
-mkdir -p llvm/libunwind
-tar -xf $UNWIND.src.tar.xz --strip-components 1 -C llvm/libunwind/
-mkdir -p llvm/tools/lld
-tar -xf $LLD.src.tar.xz --strip-components 1 -C llvm/tools/lld/
-mkdir -p llvm/projects/compiler-rt
-tar -xf $COMPILER_RT.src.tar.xz --strip-components 1 -C llvm/projects/compiler-rt/
-
-#llvm has cmake/modules, where compiler-rt has cmake/Modules and cmake.sr.tar.xz has cmake/Modules...
-tar -xf $CMAKE_MOD.src.tar.xz --strip-components 2 -C llvm/cmake/modules/
-tar -xf $CMAKE_MOD.src.tar.xz --strip-components 1 -C llvm/projects/compiler-rt/cmake/
+echo "Cloning LLVM (tag $LLVM_TAG only)"
+# clone -b can also accept tag names
+git clone [REDACTED_PROFILE_URL]/llvm-project llvm-src -b ${LLVM_TAG} --single-branch --depth 1 -c advice.detachedHead=false
+pushd llvm-src
+git -c advice.detachedHead=false checkout tags/${LLVM_TAG} -b ${LLVM_BRANCH}
+popd
+
 
 # this fixes an issue where AT_HWCAP2 is just not defined correctly in our sysroot. This is likely due to
 # AT_HWCAP2 being around since glibc 2.18 offically, while we are still stuck on 2.17 glibc.
-patch -d llvm/projects/compiler-rt/ -p 2 < /src/patches/compiler-rt/manually-define-AT_HWCAP2.diff
+patch -d llvm-src -p 1 < /src/patches/compiler-rt/manually-define-AT_HWCAP2.diff
+
+# this fixes lack of HWCAP_CRC32 in the old glibc (similar issue)
+patch -d llvm-src -p 1 < /src/patches/compiler-rt/cpu_model_define_HWCAP_CRC32.diff
 
 # move back to defaulting to dwarf 4, as if we leave to dwarf 5 libs built with dwarf5 will force everything to dwarf5
 # even if you request dwarf 4. dwarf 5 currently causes issues with dump_syms and gdb/lldb earlier versions
-patch -d llvm/tools/ -p 1 < /src/patches/clang/default-dwarf-4.patch
+patch -d llvm-src -p 1 < /src/patches/clang/default-dwarf-4.patch
+
+# add a patch to disable auto-upgrade of debug info. It missed clang 16.x, so this patch shouldn't be needed for clang 17.x going forward
+# See https://reviews.llvm.org/D143229 for context
+patch -d llvm-src -p 1 < /src/patches/llvm/disable-auto-upgrade-debug-info.patch
 
 # LLVM has just failed to support stand-alone LLD build, cheat by moving a required header into a location it can be found easily
 # if you fulling include this you end up breaking other things. [REDACTED_PROFILE_URL]/llvm-project/issues/48572
-cp -rf llvm/libunwind/include/mach-o/ llvm/include
+cp -rf llvm-src/libunwind/include/mach-o/ llvm-src/llvm/include
 
 mkdir -p build-clang
 pushd build-clang
@@ -117,7 +106,8 @@ pushd build-clang
 	#   4.8.5.  You can temporarily opt out using
 	#   LLVM_TEMPORARILY_ALLOW_OLD_TOOLCHAIN, but very soon your toolchain won't be
 	#   supported.
-	cmake3 -G "Unix Makefiles" ../llvm \
+	cmake3 -G "Unix Makefiles" ../llvm-src/llvm \
+		-DLLVM_ENABLE_PROJECTS=llvm\;clang\;lld\;compiler-rt \
 		-DCMAKE_BUILD_TYPE=Release \
 		-DLLVM_ENABLE_TERMINFO=OFF \
 		-DLLVM_ENABLE_LIBXML2=OFF \
@@ -181,16 +171,16 @@ for arch in $TARGETS; do
 		continue
 	fi
 
-	mkdir -p ${OutputDirLinux}/$arch/lib/clang/${LLVM_VERSION}/{lib,share,include}
+	mkdir -p ${OutputDirLinux}/$arch/lib/clang/${LLVM_VERSION_MAJOR}/{lib,share,include}
 
 	# copy share + include files (same as x86_64)
-	cp -r ${OutputDirLinux}/x86_64-unknown-linux-gnu/lib/clang/${LLVM_VERSION}/share/* ${OutputDirLinux}/$arch/lib/clang/${LLVM_VERSION}/share/
-	cp -r ${OutputDirLinux}/x86_64-unknown-linux-gnu/lib/clang/${LLVM_VERSION}/include/* ${OutputDirLinux}/$arch/lib/clang/${LLVM_VERSION}/include/
+	cp -r ${OutputDirLinux}/x86_64-unknown-linux-gnu/lib/clang/${LLVM_VERSION_MAJOR}/share/* ${OutputDirLinux}/$arch/lib/clang/${LLVM_VERSION_MAJOR}/share/
+	cp -r ${OutputDirLinux}/x86_64-unknown-linux-gnu/lib/clang/${LLVM_VERSION_MAJOR}/include/* ${OutputDirLinux}/$arch/lib/clang/${LLVM_VERSION_MAJOR}/include/
 
 	mkdir -p build-rt-$arch
 	pushd build-rt-$arch
 
-		cmake3 -G "Unix Makefiles" ../llvm/projects/compiler-rt \
+		cmake3 -G "Unix Makefiles" ../llvm-src/compiler-rt \
 			-DCMAKE_BUILD_TYPE=Release \
 			-DCMAKE_SYSTEM_NAME="Linux" \
 			-DCOMPILER_RT_DEFAULT_TARGET_ONLY=ON \
@@ -219,7 +209,7 @@ for arch in $TARGETS; do
 	popd
 
 	echo "Copying compiler rt..."
-	cp -r install-rt-$arch/lib/* ${OutputDirLinux}/$arch/lib/clang/${LLVM_VERSION}/lib/
+	cp -r install-rt-$arch/lib/* ${OutputDirLinux}/$arch/lib/clang/${LLVM_VERSION_MAJOR}/lib/
 done
 
 # Create version file
@@ -261,7 +251,8 @@ pushd ${OutputDirLinux}
 	mkdir -p build/{src,scripts}
 	cp /src/build/build-linux-x86_64-unknown-linux-gnu/.build/tarballs/* build/src
 	cp /src/build/build-linux-aarch64-unknown-linux-gnueabi/.build/tarballs/* build/src
-	cp /src/build/*.src.tar.xz build/src
+	tar czfvh /src/build/llvm-${LLVM_VERSION}-github-snapshot.src.tar.gz --hard-dereference /src/build/llvm-src
+	cp /src/build/*.src.tar.gz build/src
 	cp /src/*.{config,sh,nsi,bat} build/scripts
 
 	echo tar czfhv /src/build/${ToolChainVersionName}.tar.gz --hard-dereference *
@@ -273,10 +264,12 @@ pushd ${OutputDirWindows}
 	mkdir -p build/{src,scripts}
 	cp /src/build/build-windows-x86_64-unknown-linux-gnu/.build/tarballs/* build/src
 	cp /src/build/build-windows-aarch64-unknown-linux-gnueabi/.build/tarballs/* build/src
-	cp /src/build/*.src.tar.xz build/src
+	zip -r /src/build/llvm-${LLVM_VERSION}-github-snapshot.src.zip /src/build/llvm-src
+	cp /src/build/*.src.zip build/src
 	cp /src/*.{config,sh,nsi,bat} build/scripts
 
 	zip -r /src/build/${ToolChainVersionName}-windows.zip *
 popd
 
+echo "Remember to repack Linux tarball - it needs DSOs from under /lib directory deleted or clang won't run! Compare with previous toolchains! ([RCL] FIXME: 2023-06-28 fix the script to do this)"
 echo done.
