PR #10620: Make sure BuildU[A|B]T.bat explicitly call the Windows find.exe utility

diff --git a/Engine/Build/BatchFiles/BuildUAT.bat b/Engine/Build/BatchFiles/BuildUAT.bat
--- a/Engine/Build/BatchFiles/BuildUAT.bat
+++ b/Engine/Build/BatchFiles/BuildUAT.bat
@@ -1,131 +1,134 @@
-@echo off
-
-rem ## Unreal Engine AutomationTool build script
-rem ## Copyright Epic Games, Inc. All Rights Reserved.
-rem ##
-rem ## This script is expecting to exist in the Engine/Build/BatchFiles directory.  It will not work correctly
-rem ## if you copy it to a different location and run it.
-
-setlocal
-
-rem ## First, make sure the batch file exists in the folder we expect it to.  This is necessary in order to
-rem ## verify that our relative path to the /Engine/Source directory is correct
-if not exist "%~dp0..\..\Source" goto Error_BatchFileInWrongLocation
-
-rem ## Change the CWD to /Engine/Source.
-pushd "%~dp0..\..\Source"
-if not exist ..\Build\BatchFiles\BuildUAT.bat goto Error_BatchFileInWrongLocation
-
-rem Check to see if the files in the AutomationTool, EpicGames.Build, EpicGames.Core, or UnrealBuildTool directory have changed.
-rem find ".cs" files to only lines that match those names - excludes lines that will change for uninteresting reasons, like free space
-md ..\Intermediate\Build >nul 2>nul
-
-dir /s ^
- Programs\Shared\EpicGames.Box\*.cs ^
- Programs\Shared\EpicGames.Box\*.csproj ^
- Programs\Shared\EpicGames.Build\*.cs ^
- Programs\Shared\EpicGames.Build\*.csproj ^
- Programs\Shared\EpicGames.Core\*.cs ^
- Programs\Shared\EpicGames.Core\*.csproj ^
- Programs\Shared\EpicGames.Horde\*.cs ^
- Programs\Shared\EpicGames.Horde\*.csproj ^
- Programs\Shared\EpicGames.IoHash\*.cs ^
- Programs\Shared\EpicGames.IoHash\*.csproj ^
- Programs\Shared\EpicGames.MsBuild\*.cs ^
- Programs\Shared\EpicGames.MsBuild\*.csproj ^
- Programs\Shared\EpicGames.OIDC\*.cs ^
- Programs\Shared\EpicGames.OIDC\*.csproj ^
- Programs\Shared\EpicGames.Serialization\*.cs ^
- Programs\Shared\EpicGames.Serialization\*.csproj ^
- Programs\Shared\EpicGames.UHT\*.cs ^
- Programs\Shared\EpicGames.UHT\*.csproj ^
- Programs\UnrealBuildTool\*.cs ^
- Programs\UnrealBuildTool\*.csproj ^
- | find ".cs" > ..\Intermediate\Build\AutomationToolFiles.txt
-
-if not exist ..\Platforms goto NoPlatforms
-for /d %%D in (..\Platforms\*) do (
-	if exist %%D\Source\Programs\UnrealBuildTool (
-		dir /s ^
-		 %%D\Source\Programs\AutomationTool\*.cs ^
-		 %%D\Source\Programs\AutomationTool\*.csproj ^
-		 %%D\Source\Programs\UnrealBuildTool\*.cs ^
-		 %%D\Source\Programs\UnrealBuildTool\*.csproj ^
-		 | find ".cs" >> ..\Intermediate\Build\AutomationToolFiles.txt
-	) 2>nul
-)
-:NoPlatforms
-
-if not exist ..\Restricted goto NoRestricted
-for /d %%D in (..\Restricted\*) do (
-	if exist %%D\Source\Programs\UnrealBuildTool (
-		dir /s ^
-		 %%D\Source\Programs\AutomationTool\*.cs ^
-		 %%D\Source\Programs\AutomationTool\*.csproj ^
-		 %%D\Source\Programs\UnrealBuildTool\*.cs ^
-		 %%D\Source\Programs\UnrealBuildTool\*.csproj ^
-		 | find ".cs" >> ..\Intermediate\Build\AutomationToolFiles.txt
-	) 2>nul
-)
-:NoRestricted
-
-rem note: no /s
-dir ^
- Programs\Shared\MetaData.cs ^
- Programs\AutomationTool\*.cs ^
- Programs\AutomationTool\*.csproj ^
- | find ".cs" >>..\Intermediate\Build\AutomationToolFiles.txt
-
-set MSBUILD_LOGLEVEL=%1
-if not defined %MSBUILD_LOGLEVEL set MSBUILD_LOGLEVEL=quiet
-
-set ARGUMENT=%2
-if not defined %ARGUMENT goto Check_UpToDate
-if /I "%ARGUMENT%" == "FORCE" goto Build_AutomationTool
-
-:Check_UpToDate
-if not exist ..\Binaries\DotNET\AutomationTool\AutomationTool.dll goto Build_AutomationTool
-set RUNUAT_EXITCODE=0
-rem per https://ss64.com/nt/fc.html using redirection syntax rather than errorlevel, based on observed inconsistent results from this function
-fc ..\Intermediate\Build\AutomationToolFiles.txt ..\Intermediate\Build\AutomationToolPrevFiles.txt >nul 2>&1 && goto Exit
-
-:Build_AutomationTool
-rem ## Verify that dotnet is present
-call "%~dp0GetDotnetPath.bat"
-if errorlevel 1 goto Error_NoDotnetSDK
-
-echo Building AutomationTool...
-dotnet build Programs\AutomationTool\AutomationTool.csproj -c Development -v %MSBUILD_LOGLEVEL%
-if errorlevel 1 goto Error_UATCompileFailed
-
-rem record input files - regardless of how we got here, these are now our point of reference
-copy /y ..\Intermediate\Build\AutomationToolFiles.txt ..\Intermediate\Build\AutomationToolPrevFiles.txt >nul
-
-goto Exit
-
-
-:Error_BatchFileInWrongLocation
-echo.
-echo BuildUAT ERROR: The batch file does not appear to be located in the /Engine/Build/BatchFiles directory.  This script must be run from within that directory.
-echo.
-set RUNUAT_EXITCODE=1
-goto Exit
-
-:Error_NoDotnetSDK
-echo.
-echo RunUBT ERROR: Unable to find a install of Dotnet SDK.  Please make sure you have it installed and that `dotnet` is a globally available command.
-echo.
-set RUNUAT_EXITCODE=1
-goto Exit
-
-:Error_UATCompileFailed
-echo.
-echo RunUBT ERROR: UnrealBuildTool failed to compile.
-echo.
-set RUNUAT_EXITCODE=1
-goto Exit
-
-:Exit
-rem ## Restore original CWD in case we change it
-popd
-exit /B %RUNUAT_EXITCODE%
+@echo off
+
+rem ## Unreal Engine AutomationTool build script
+rem ## Copyright Epic Games, Inc. All Rights Reserved.
+rem ##
+rem ## This script is expecting to exist in the Engine/Build/BatchFiles directory.  It will not work correctly
+rem ## if you copy it to a different location and run it.
+
+setlocal
+
+rem ## Make sure we use the Windows "find" utility and not a Unix-like tool found in PATH
+set FIND=%SYSTEMROOT%\System32\find.exe
+
+rem ## First, make sure the batch file exists in the folder we expect it to.  This is necessary in order to
+rem ## verify that our relative path to the /Engine/Source directory is correct
+if not exist "%~dp0..\..\Source" goto Error_BatchFileInWrongLocation
+
+rem ## Change the CWD to /Engine/Source.
+pushd "%~dp0..\..\Source"
+if not exist ..\Build\BatchFiles\BuildUAT.bat goto Error_BatchFileInWrongLocation
+
+rem Check to see if the files in the AutomationTool, EpicGames.Build, EpicGames.Core, or UnrealBuildTool directory have changed.
+rem find ".cs" files to only lines that match those names - excludes lines that will change for uninteresting reasons, like free space
+md ..\Intermediate\Build >nul 2>nul
+
+dir /s ^
+ Programs\Shared\EpicGames.Box\*.cs ^
+ Programs\Shared\EpicGames.Box\*.csproj ^
+ Programs\Shared\EpicGames.Build\*.cs ^
+ Programs\Shared\EpicGames.Build\*.csproj ^
+ Programs\Shared\EpicGames.Core\*.cs ^
+ Programs\Shared\EpicGames.Core\*.csproj ^
+ Programs\Shared\EpicGames.Horde\*.cs ^
+ Programs\Shared\EpicGames.Horde\*.csproj ^
+ Programs\Shared\EpicGames.IoHash\*.cs ^
+ Programs\Shared\EpicGames.IoHash\*.csproj ^
+ Programs\Shared\EpicGames.MsBuild\*.cs ^
+ Programs\Shared\EpicGames.MsBuild\*.csproj ^
+ Programs\Shared\EpicGames.OIDC\*.cs ^
+ Programs\Shared\EpicGames.OIDC\*.csproj ^
+ Programs\Shared\EpicGames.Serialization\*.cs ^
+ Programs\Shared\EpicGames.Serialization\*.csproj ^
+ Programs\Shared\EpicGames.UHT\*.cs ^
+ Programs\Shared\EpicGames.UHT\*.csproj ^
+ Programs\UnrealBuildTool\*.cs ^
+ Programs\UnrealBuildTool\*.csproj ^
+ | %FIND% ".cs" > ..\Intermediate\Build\AutomationToolFiles.txt
+
+if not exist ..\Platforms goto NoPlatforms
+for /d %%D in (..\Platforms\*) do (
+	if exist %%D\Source\Programs\UnrealBuildTool (
+		dir /s ^
+		 %%D\Source\Programs\AutomationTool\*.cs ^
+		 %%D\Source\Programs\AutomationTool\*.csproj ^
+		 %%D\Source\Programs\UnrealBuildTool\*.cs ^
+		 %%D\Source\Programs\UnrealBuildTool\*.csproj ^
+		 | %FIND% ".cs" >> ..\Intermediate\Build\AutomationToolFiles.txt
+	) 2>nul
+)
+:NoPlatforms
+
+if not exist ..\Restricted goto NoRestricted
+for /d %%D in (..\Restricted\*) do (
+	if exist %%D\Source\Programs\UnrealBuildTool (
+		dir /s ^
+		 %%D\Source\Programs\AutomationTool\*.cs ^
+		 %%D\Source\Programs\AutomationTool\*.csproj ^
+		 %%D\Source\Programs\UnrealBuildTool\*.cs ^
+		 %%D\Source\Programs\UnrealBuildTool\*.csproj ^
+		 | %FIND% ".cs" >> ..\Intermediate\Build\AutomationToolFiles.txt
+	) 2>nul
+)
+:NoRestricted
+
+rem note: no /s
+dir ^
+ Programs\Shared\MetaData.cs ^
+ Programs\AutomationTool\*.cs ^
+ Programs\AutomationTool\*.csproj ^
+ | %FIND% ".cs" >>..\Intermediate\Build\AutomationToolFiles.txt
+
+set MSBUILD_LOGLEVEL=%1
+if not defined %MSBUILD_LOGLEVEL set MSBUILD_LOGLEVEL=quiet
+
+set ARGUMENT=%2
+if not defined %ARGUMENT goto Check_UpToDate
+if /I "%ARGUMENT%" == "FORCE" goto Build_AutomationTool
+
+:Check_UpToDate
+if not exist ..\Binaries\DotNET\AutomationTool\AutomationTool.dll goto Build_AutomationTool
+set RUNUAT_EXITCODE=0
+rem per https://ss64.com/nt/fc.html using redirection syntax rather than errorlevel, based on observed inconsistent results from this function
+fc ..\Intermediate\Build\AutomationToolFiles.txt ..\Intermediate\Build\AutomationToolPrevFiles.txt >nul 2>&1 && goto Exit
+
+:Build_AutomationTool
+rem ## Verify that dotnet is present
+call "%~dp0GetDotnetPath.bat"
+if errorlevel 1 goto Error_NoDotnetSDK
+
+echo Building AutomationTool...
+dotnet build Programs\AutomationTool\AutomationTool.csproj -c Development -v %MSBUILD_LOGLEVEL%
+if errorlevel 1 goto Error_UATCompileFailed
+
+rem record input files - regardless of how we got here, these are now our point of reference
+copy /y ..\Intermediate\Build\AutomationToolFiles.txt ..\Intermediate\Build\AutomationToolPrevFiles.txt >nul
+
+goto Exit
+
+
+:Error_BatchFileInWrongLocation
+echo.
+echo BuildUAT ERROR: The batch file does not appear to be located in the /Engine/Build/BatchFiles directory.  This script must be run from within that directory.
+echo.
+set RUNUAT_EXITCODE=1
+goto Exit
+
+:Error_NoDotnetSDK
+echo.
+echo RunUBT ERROR: Unable to find a install of Dotnet SDK.  Please make sure you have it installed and that `dotnet` is a globally available command.
+echo.
+set RUNUAT_EXITCODE=1
+goto Exit
+
+:Error_UATCompileFailed
+echo.
+echo RunUBT ERROR: UnrealBuildTool failed to compile.
+echo.
+set RUNUAT_EXITCODE=1
+goto Exit
+
+:Exit
+rem ## Restore original CWD in case we change it
+popd
+exit /B %RUNUAT_EXITCODE%
diff --git a/Engine/Build/BatchFiles/BuildUBT.bat b/Engine/Build/BatchFiles/BuildUBT.bat
--- a/Engine/Build/BatchFiles/BuildUBT.bat
+++ b/Engine/Build/BatchFiles/BuildUBT.bat
@@ -1,125 +1,128 @@
-@echo off
-
-rem ## Unreal Engine UnrealBuildTool build script
-rem ## Copyright Epic Games, Inc. All Rights Reserved.
-rem ##
-rem ## This script is expecting to exist in the Engine/Build/BatchFiles directory.  It will not work correctly
-rem ## if you copy it to a different location and run it.
-
-setlocal
-
-rem ## First, make sure the batch file exists in the folder we expect it to.  This is necessary in order to
-rem ## verify that our relative path to the /Engine/Source directory is correct
-if not exist "%~dp0..\..\Source" goto Error_BatchFileInWrongLocation
-
-rem ## Change the CWD to /Engine/Source.
-pushd "%~dp0..\..\Source"
-if not exist ..\Build\BatchFiles\BuildUBT.bat goto Error_BatchFileInWrongLocation
-
-rem Check to see if the files in the UnrealBuildTool or Shared project directory have changed.
-rem find ".cs" files to only lines that match those names - excludes lines that will change for uninteresting reasons, like free space
-md ..\Intermediate\Build >nul 2>nul
-
-dir /s ^
- Programs\Shared\EpicGames.Box\*.cs ^
- Programs\Shared\EpicGames.Box\*.csproj ^
- Programs\Shared\EpicGames.Build\*.cs ^
- Programs\Shared\EpicGames.Build\*.csproj ^
- Programs\Shared\EpicGames.Core\*.cs ^
- Programs\Shared\EpicGames.Core\*.csproj ^
- Programs\Shared\EpicGames.Horde\*.cs ^
- Programs\Shared\EpicGames.Horde\*.csproj ^
- Programs\Shared\EpicGames.IoHash\*.cs ^
- Programs\Shared\EpicGames.IoHash\*.csproj ^
- Programs\Shared\EpicGames.MsBuild\*.cs ^
- Programs\Shared\EpicGames.MsBuild\*.csproj ^
- Programs\Shared\EpicGames.OIDC\*.cs ^
- Programs\Shared\EpicGames.OIDC\*.csproj ^
- Programs\Shared\EpicGames.Serialization\*.cs ^
- Programs\Shared\EpicGames.Serialization\*.csproj ^
- Programs\Shared\EpicGames.UHT\*.cs ^
- Programs\Shared\EpicGames.UHT\*.csproj ^
- Programs\UnrealBuildTool\*.cs ^
- Programs\UnrealBuildTool\*.csproj ^
- | find ".cs" > ..\Intermediate\Build\UnrealBuildToolFiles.txt
-
-if not exist ..\Platforms goto NoPlatforms
-for /d %%D in (..\Platforms\*) do (
-	if exist %%D\Source\Programs\UnrealBuildTool (
-		dir /s ^
-		 %%D\Source\Programs\UnrealBuildTool\*.cs ^
-		 %%D\Source\Programs\UnrealBuildTool\*.csproj ^
-		 | find ".cs" >> ..\Intermediate\Build\UnrealBuildToolFiles.txt
-	) 2>nul
-)
-:NoPlatforms
-
-if not exist ..\Restricted goto NoRestricted
-for /d %%D in (..\Restricted\*) do (
-	if exist %%D\Source\Programs\UnrealBuildTool (
-		dir /s ^
-		 %%D\Source\Programs\UnrealBuildTool\*.cs ^
-		 %%D\Source\Programs\UnrealBuildTool\*.csproj ^
-		 | find ".cs" >> ..\Intermediate\Build\UnrealBuildToolFiles.txt
-	) 2>nul
-)
-:NoRestricted
-
-rem note: no /s
-dir ^
- Programs\Shared\MetaData.cs ^
- | find ".cs" >>..\Intermediate\Build\UnrealBuildToolFiles.txt
-
-set MSBUILD_LOGLEVEL=%1
-if not defined %MSBUILD_LOGLEVEL set MSBUILD_LOGLEVEL=quiet
-
-set ARGUMENT=%2
-if not defined %ARGUMENT goto Check_UpToDate
-if /I "%ARGUMENT%" == "FORCE" goto Build_UnrealBuildTool
-
-:Check_UpToDate
-if not exist ..\Binaries\DotNET\UnrealBuildTool\UnrealBuildTool.dll goto Build_UnrealBuildTool
-set RUNUBT_EXITCODE=0
-rem per https://ss64.com/nt/fc.html using redirection syntax rather than errorlevel, based on observed inconsistent results from this function
-fc ..\Intermediate\Build\UnrealBuildToolFiles.txt ..\Intermediate\Build\UnrealBuildToolPrevFiles.txt >nul 2>&1 && goto Exit
-
-:Build_UnrealBuildTool
-rem ## Verify that dotnet is present
-call "%~dp0GetDotnetPath.bat"
-if errorlevel 1 goto Error_NoDotnetSDK
-
-echo Building UnrealBuildTool...
-dotnet build Programs\UnrealBuildTool\UnrealBuildTool.csproj -c Development -v %MSBUILD_LOGLEVEL%
-if errorlevel 1 goto Error_UBTCompileFailed
-
-rem record input files - regardless of how we got here, these are now our point of reference
-copy /y ..\Intermediate\Build\UnrealBuildToolFiles.txt ..\Intermediate\Build\UnrealBuildToolPrevFiles.txt >nul
-
-goto Exit
-
-
-:Error_BatchFileInWrongLocation
-echo.
-echo BuildUBT ERROR: The batch file does not appear to be located in the /Engine/Build/BatchFiles directory.  This script must be run from within that directory.
-echo.
-set RUNUBT_EXITCODE=1
-goto Exit
-
-:Error_NoDotnetSDK
-echo.
-echo RunUBT ERROR: Unable to find a install of Dotnet SDK.  Please make sure you have it installed and that `dotnet` is a globally available command.
-echo.
-set RUNUBT_EXITCODE=1
-goto Exit
-
-:Error_UBTCompileFailed
-echo.
-echo RunUBT ERROR: UnrealBuildTool failed to compile.
-echo.
-set RUNUBT_EXITCODE=1
-goto Exit
-
-:Exit
-rem ## Restore original CWD in case we change it
-popd
-exit /B %RUNUBT_EXITCODE%
+@echo off
+
+rem ## Unreal Engine UnrealBuildTool build script
+rem ## Copyright Epic Games, Inc. All Rights Reserved.
+rem ##
+rem ## This script is expecting to exist in the Engine/Build/BatchFiles directory.  It will not work correctly
+rem ## if you copy it to a different location and run it.
+
+setlocal
+
+rem ## Make sure we use the Windows "find" utility and not a Unix-like tool found in PATH
+set FIND=%SYSTEMROOT%\System32\find.exe
+
+rem ## First, make sure the batch file exists in the folder we expect it to.  This is necessary in order to
+rem ## verify that our relative path to the /Engine/Source directory is correct
+if not exist "%~dp0..\..\Source" goto Error_BatchFileInWrongLocation
+
+rem ## Change the CWD to /Engine/Source.
+pushd "%~dp0..\..\Source"
+if not exist ..\Build\BatchFiles\BuildUBT.bat goto Error_BatchFileInWrongLocation
+
+rem Check to see if the files in the UnrealBuildTool or Shared project directory have changed.
+rem find ".cs" files to only lines that match those names - excludes lines that will change for uninteresting reasons, like free space
+md ..\Intermediate\Build >nul 2>nul
+
+dir /s ^
+ Programs\Shared\EpicGames.Box\*.cs ^
+ Programs\Shared\EpicGames.Box\*.csproj ^
+ Programs\Shared\EpicGames.Build\*.cs ^
+ Programs\Shared\EpicGames.Build\*.csproj ^
+ Programs\Shared\EpicGames.Core\*.cs ^
+ Programs\Shared\EpicGames.Core\*.csproj ^
+ Programs\Shared\EpicGames.Horde\*.cs ^
+ Programs\Shared\EpicGames.Horde\*.csproj ^
+ Programs\Shared\EpicGames.IoHash\*.cs ^
+ Programs\Shared\EpicGames.IoHash\*.csproj ^
+ Programs\Shared\EpicGames.MsBuild\*.cs ^
+ Programs\Shared\EpicGames.MsBuild\*.csproj ^
+ Programs\Shared\EpicGames.OIDC\*.cs ^
+ Programs\Shared\EpicGames.OIDC\*.csproj ^
+ Programs\Shared\EpicGames.Serialization\*.cs ^
+ Programs\Shared\EpicGames.Serialization\*.csproj ^
+ Programs\Shared\EpicGames.UHT\*.cs ^
+ Programs\Shared\EpicGames.UHT\*.csproj ^
+ Programs\UnrealBuildTool\*.cs ^
+ Programs\UnrealBuildTool\*.csproj ^
+ | %FIND% ".cs" > ..\Intermediate\Build\UnrealBuildToolFiles.txt
+
+if not exist ..\Platforms goto NoPlatforms
+for /d %%D in (..\Platforms\*) do (
+	if exist %%D\Source\Programs\UnrealBuildTool (
+		dir /s ^
+		 %%D\Source\Programs\UnrealBuildTool\*.cs ^
+		 %%D\Source\Programs\UnrealBuildTool\*.csproj ^
+		 | %FIND% ".cs" >> ..\Intermediate\Build\UnrealBuildToolFiles.txt
+	) 2>nul
+)
+:NoPlatforms
+
+if not exist ..\Restricted goto NoRestricted
+for /d %%D in (..\Restricted\*) do (
+	if exist %%D\Source\Programs\UnrealBuildTool (
+		dir /s ^
+		 %%D\Source\Programs\UnrealBuildTool\*.cs ^
+		 %%D\Source\Programs\UnrealBuildTool\*.csproj ^
+		 | %FIND% ".cs" >> ..\Intermediate\Build\UnrealBuildToolFiles.txt
+	) 2>nul
+)
+:NoRestricted
+
+rem note: no /s
+dir ^
+ Programs\Shared\MetaData.cs ^
+ | %FIND% ".cs" >>..\Intermediate\Build\UnrealBuildToolFiles.txt
+
+set MSBUILD_LOGLEVEL=%1
+if not defined %MSBUILD_LOGLEVEL set MSBUILD_LOGLEVEL=quiet
+
+set ARGUMENT=%2
+if not defined %ARGUMENT goto Check_UpToDate
+if /I "%ARGUMENT%" == "FORCE" goto Build_UnrealBuildTool
+
+:Check_UpToDate
+if not exist ..\Binaries\DotNET\UnrealBuildTool\UnrealBuildTool.dll goto Build_UnrealBuildTool
+set RUNUBT_EXITCODE=0
+rem per https://ss64.com/nt/fc.html using redirection syntax rather than errorlevel, based on observed inconsistent results from this function
+fc ..\Intermediate\Build\UnrealBuildToolFiles.txt ..\Intermediate\Build\UnrealBuildToolPrevFiles.txt >nul 2>&1 && goto Exit
+
+:Build_UnrealBuildTool
+rem ## Verify that dotnet is present
+call "%~dp0GetDotnetPath.bat"
+if errorlevel 1 goto Error_NoDotnetSDK
+
+echo Building UnrealBuildTool...
+dotnet build Programs\UnrealBuildTool\UnrealBuildTool.csproj -c Development -v %MSBUILD_LOGLEVEL%
+if errorlevel 1 goto Error_UBTCompileFailed
+
+rem record input files - regardless of how we got here, these are now our point of reference
+copy /y ..\Intermediate\Build\UnrealBuildToolFiles.txt ..\Intermediate\Build\UnrealBuildToolPrevFiles.txt >nul
+
+goto Exit
+
+
+:Error_BatchFileInWrongLocation
+echo.
+echo BuildUBT ERROR: The batch file does not appear to be located in the /Engine/Build/BatchFiles directory.  This script must be run from within that directory.
+echo.
+set RUNUBT_EXITCODE=1
+goto Exit
+
+:Error_NoDotnetSDK
+echo.
+echo RunUBT ERROR: Unable to find a install of Dotnet SDK.  Please make sure you have it installed and that `dotnet` is a globally available command.
+echo.
+set RUNUBT_EXITCODE=1
+goto Exit
+
+:Error_UBTCompileFailed
+echo.
+echo RunUBT ERROR: UnrealBuildTool failed to compile.
+echo.
+set RUNUBT_EXITCODE=1
+goto Exit
+
+:Exit
+rem ## Restore original CWD in case we change it
+popd
+exit /B %RUNUBT_EXITCODE%
