- Added a mkdir to create the directory before writing a file

diff --git a/Engine/Build/BatchFiles/Mac/UpdateVersionAfterBuild.sh b/Engine/Build/BatchFiles/Mac/UpdateVersionAfterBuild.sh
--- a/Engine/Build/BatchFiles/Mac/UpdateVersionAfterBuild.sh
+++ b/Engine/Build/BatchFiles/Mac/UpdateVersionAfterBuild.sh
@@ -48,6 +48,7 @@ fi
 # now write out an xcconfig containing all platform build versions relative to the project
 XCCONFIG_FILE="$1/Intermediate/Build/Versions.xcconfig"
 
+mkdir -p "$1/Intermediate/Build"
 echo "UE_MAC_BUILD_VERSION = $MAC_VERSION" > "$XCCONFIG_FILE"
 echo "UE_IOS_BUILD_VERSION = $IOS_VERSION" >> "$XCCONFIG_FILE"
 echo "UE_TVOS_BUILD_VERSION = $TVOS_VERSION" >> "$XCCONFIG_FILE"
