Download v22 native Linux toolchain when setting up Linux.

diff --git a/Engine/Build/BatchFiles/Linux/SetupToolchain.sh b/Engine/Build/BatchFiles/Linux/SetupToolchain.sh
--- a/Engine/Build/BatchFiles/Linux/SetupToolchain.sh
+++ b/Engine/Build/BatchFiles/Linux/SetupToolchain.sh
@@ -11,7 +11,7 @@ if [ -z $GIT_DIR ]; then
 	fi
 fi
 
-TOOLCHAIN_VERSION=v21_clang-15.0.1-centos7
+TOOLCHAIN_VERSION=v22_clang-16.0.6-centos7
 TOOLCHAIN_ARCHIVE=$TOOLCHAIN_VERSION.tar.gz
 
 TOOLCHAIN_URL=http://cdn.unrealengine.com/Toolchain_Linux/native-linux-$TOOLCHAIN_ARCHIVE
