Removing uneeded call to SetupDotnet.sh on mac

#rnx

diff --git a/Engine/Build/BatchFiles/Mac/GitDependencies.sh b/Engine/Build/BatchFiles/Mac/GitDependencies.sh
--- a/Engine/Build/BatchFiles/Mac/GitDependencies.sh
+++ b/Engine/Build/BatchFiles/Mac/GitDependencies.sh
@@ -10,8 +10,6 @@ fi
 
 cd "$(dirname "$SCRIPT_PATH")" && SCRIPT_PATH="`pwd`/$(basename "$SCRIPT_PATH")"
 
-"$(dirname "$SCRIPT_PATH")/SetupDotnet.sh"
-
 cd ../../../..
 
 # Select the preferred architecture for the current system
