Fix Linux toolchain building scripts.

- Native Linux toolchain is produced in already a runnable state with the proper name for distribution.
- Both cross and native toolchains have their repo path changed to the github repo they are cloned from.
- Not building bolt as this for some reason cases llvm-profdata and symbolizer not to be added to the install folder.

diff --git a/Engine/Build/BatchFiles/Linux/Toolchain/RunMe.bat b/Engine/Build/BatchFiles/Linux/Toolchain/RunMe.bat
--- a/Engine/Build/BatchFiles/Linux/Toolchain/RunMe.bat
+++ b/Engine/Build/BatchFiles/Linux/Toolchain/RunMe.bat
@@ -59,7 +59,7 @@ popd
 echo Building LLVM, clang, lld, bolt
 mkdir build_all
 pushd build_all
-%CMAKE_BINARY% -G %VS_VERSION% -DLLVM_ENABLE_PROJECTS=llvm;clang;lld;bolt -DLLVM_INTEGRATED_CRT_ALLOC=../rpmalloc -DLLVM_USE_CRT_RELEASE=MT -DCMAKE_BUILD_TYPE=Release -DCMAKE_C_FLAGS_RELEASE="/O2 /Ob3 /DNDEBUG /Zi /Gy" -DCMAKE_CXX_FLAGS_RELEASE="/O2 /Ob3 /DNDEBUG /Zi /Gy" -DCMAKE_EXE_LINKER_FLAGS_RELEASE="/DEBUG /INCREMENTAL:NO /OPT:REF /OPT:ICF" -DCMAKE_INSTALL_PREFIX="..\install" -DLLVM_INCLUDE_BENCHMARKS=OFF -DLLVM_TARGETS_TO_BUILD="AArch64;X86" -DLLVM_ENABLE_ZLIB=FORCE_ON -DZLIB_LIBRARY="%ZLIB_PATH%/lib/Win64-llvm/Release/zlibstatic.lib" -DZLIB_INCLUDE_DIR="%ZLIB_PATH%/include/Win64/VS2015"  "..\source\llvm"
+%CMAKE_BINARY% -G %VS_VERSION% -DLLVM_ENABLE_PROJECTS=llvm;clang;lld -DLLVM_INTEGRATED_CRT_ALLOC=../rpmalloc -DLLVM_USE_CRT_RELEASE=MT -DCMAKE_BUILD_TYPE=Release -DCMAKE_C_FLAGS_RELEASE="/O2 /Ob3 /DNDEBUG /Zi /Gy" -DCMAKE_CXX_FLAGS_RELEASE="/O2 /Ob3 /DNDEBUG /Zi /Gy" -DCMAKE_EXE_LINKER_FLAGS_RELEASE="/DEBUG /INCREMENTAL:NO /OPT:REF /OPT:ICF" -DCMAKE_INSTALL_PREFIX="..\install" -DLLVM_INCLUDE_BENCHMARKS=OFF -DLLVM_TARGETS_TO_BUILD="AArch64;X86" -DLLVM_ENABLE_ZLIB=FORCE_ON -DZLIB_LIBRARY="%ZLIB_PATH%/lib/Win64-llvm/Release/zlibstatic.lib" -DZLIB_INCLUDE_DIR="%ZLIB_PATH%/include/Win64/VS2015" -DCLANG_REPOSITORY_STRING="github.com/llvm/llvm-project" "..\source\llvm"
 %CMAKE_BINARY% --build . --target install --config Release
 popd
 
diff --git a/Engine/Build/BatchFiles/Linux/Toolchain/build_linux_toolchain.sh b/Engine/Build/BatchFiles/Linux/Toolchain/build_linux_toolchain.sh
--- a/Engine/Build/BatchFiles/Linux/Toolchain/build_linux_toolchain.sh
+++ b/Engine/Build/BatchFiles/Linux/Toolchain/build_linux_toolchain.sh
@@ -118,7 +118,8 @@ pushd build-clang
 		-DLLVM_TEMPORARILY_ALLOW_OLD_TOOLCHAIN=ON \
 		-DCMAKE_INSTALL_PREFIX=${InstallClangDir} \
 		-DLLVM_INCLUDE_BENCHMARKS=OFF \
-		-DLLVM_TARGETS_TO_BUILD="AArch64;X86"
+		-DLLVM_TARGETS_TO_BUILD="AArch64;X86" \
+		-DCLANG_REPOSITORY_STRING="github.com/llvm/llvm-project"
 
 	make -j$CORES && make install
 popd
@@ -255,8 +256,17 @@ pushd ${OutputDirLinux}
 	cp /src/build/*.src.tar.gz build/src
 	cp /src/*.{config,sh,nsi,bat} build/scripts
 
-	echo tar czfhv /src/build/${ToolChainVersionName}.tar.gz --hard-dereference *
-	tar czfhv /src/build/${ToolChainVersionName}.tar.gz --hard-dereference *
+	# copy the toolchain in the directory named its version as per convention
+	mkdir ${OutputDirLinux}/${ToolChainVersionName}
+	cp -r x86_64-unknown-linux-gnu ${OutputDirLinux}/${ToolChainVersionName}
+	cp -r build ${OutputDirLinux}/${ToolChainVersionName}
+	cp -r aarch64-unknown-linux-gnueabi ${OutputDirLinux}/${ToolChainVersionName}
+	cp -r ToolchainVersion.txt ${OutputDirLinux}/${ToolChainVersionName}
+
+	# delete libraries in x86_64's lib folder or bundled binares with crash
+	find ${OutputDirLinux}/${ToolChainVersionName}/x86_64-unknown-linux-gnu/lib/ -maxdepth 1 -type f -delete
+
+	tar czfhv native-linux-${ToolChainVersionName}.tar.gz --hard-dereference ${ToolChainVersionName}
 popd
 
 # Pack Windows files
@@ -271,5 +281,4 @@ pushd ${OutputDirWindows}
 	zip -r /src/build/${ToolChainVersionName}-windows.zip *
 popd
 
-echo "Remember to repack Linux tarball - it needs DSOs from under /lib directory deleted or clang won't run! Compare with previous toolchains! ([RCL] FIXME: 2023-06-28 fix the script to do this)"
 echo done.
