Replace hard-coded toolchain version from SetupToolchain.sh by fetching it from Linux_SDK.json (which should be the source of truth on toolchain versions)

diff --git a/Engine/Build/BatchFiles/Linux/SetupToolchain.sh b/Engine/Build/BatchFiles/Linux/SetupToolchain.sh
--- a/Engine/Build/BatchFiles/Linux/SetupToolchain.sh
+++ b/Engine/Build/BatchFiles/Linux/SetupToolchain.sh
@@ -1,7 +1,7 @@
 #!/bin/bash
 # Copyright Epic Games, Inc. All Rights Reserved.
 
-# put ourselves into Engine directory (two up from location of this script)
+# put ourselves into Engine directory (three up from location of this script)
 pushd "`dirname "$0"`/../../.." > /dev/null
 
 if [ -z $GIT_DIR ]; then
@@ -11,7 +11,18 @@ if [ -z $GIT_DIR ]; then
 	fi
 fi
 
-TOOLCHAIN_VERSION=v22_clang-16.0.6-centos7
+# Parse the Toolchain Version from Linux_SDK.json
+REGEX='\"MainVersion\"\s*:\s*\"(.*)\"'
+[[ $(grep MainVersion Config/Linux/Linux_SDK.json) =~ $REGEX ]]
+TOOLCHAIN_VERSION=${BASH_REMATCH[1]}
+
+# Warn if nothing was parsed from Config/Linux/Linux_SDK.json
+if [ -z "$TOOLCHAIN_VERSION" ]
+then
+	echo "Couldn't find current toolchain, check MainVersion in $(realpath "Config/Linux/Linux_SDK.json")"
+	exit
+fi
+
 TOOLCHAIN_ARCHIVE=$TOOLCHAIN_VERSION.tar.gz
 
 TOOLCHAIN_URL=http://cdn.unrealengine.com/Toolchain_Linux/native-linux-$TOOLCHAIN_ARCHIVE
