Ensure UnrealBuildTool.dll exists when cleaning projects in installed builds

diff --git a/Engine/Build/BatchFiles/Build.bat b/Engine/Build/BatchFiles/Build.bat
--- a/Engine/Build/BatchFiles/Build.bat
+++ b/Engine/Build/BatchFiles/Build.bat
@@ -28,9 +28,6 @@ call "%~dp0GetDotnetPath.bat"
 if errorlevel 1 goto Error_NoDotnetSDK
 REM ## Skip msbuild detection if using dotnet as this is done for us by dotnet-cli
 
-rem ## If this is an installed build, we don't need to rebuild UBT. Go straight to building.
-if exist ..\Build\InstalledBuild.txt goto ReadyToBuild
-
 rem ## Compile UBT if the project file exists
 :ReadyToBuildUBT
 set ProjectFile="Programs\UnrealBuildTool\UnrealBuildTool.csproj"
