Mac: Add a log incase lipo-ing fails.

diff --git a/Engine/Build/BatchFiles/Mac/GenerateUniversalDSYM.sh b/Engine/Build/BatchFiles/Mac/GenerateUniversalDSYM.sh
--- a/Engine/Build/BatchFiles/Mac/GenerateUniversalDSYM.sh
+++ b/Engine/Build/BatchFiles/Mac/GenerateUniversalDSYM.sh
@@ -60,10 +60,18 @@ if [ ${#arches[@]} -gt 1 ]; then
   
 	# lipo all dsym binaries directly to the project-named file that we'd expect
 	echo "Merging architectures '${binpaths[*]}' together into $2"
+	
 	lipo ${binpaths[*]} -create -output "$2/Contents/Resources/DWARF/${binaryname}"
+	
+	retVal=$?
+	if [ $retVal -ne 0 ]; then
+		echo "ERROR: Something went wrong with lipo. We're in trouble now."
+	fi
 
 	# Clean up
 	rm -rf "$tempdir"
+	
+	exit $retVal
 else
   echo "Using standard dsymutil because the binary was not universal..."
   dsymutil "$1" -o "$2"
