UBA: Properly rebuild if input dlls have changed

#rnx

diff --git a/Engine/Build/BatchFiles/BuildUAT.bat b/Engine/Build/BatchFiles/BuildUAT.bat
--- a/Engine/Build/BatchFiles/BuildUAT.bat
+++ b/Engine/Build/BatchFiles/BuildUAT.bat
@@ -46,6 +46,15 @@ dir /s ^
  Programs\UnrealBuildTool\*.csproj ^
  | %FIND% ".cs" > ..\Intermediate\Build\AutomationToolFiles.txt
 
+if exist Binaries\Win64\UnrealBuildAccelerator (
+	dir /s ^
+	 Binaries\Win64\UnrealBuildAccelerator\*.dll ^
+	 | %FIND% ".dll" >> ..\Intermediate\Build\AutomationToolFiles.txt
+	dir /s ^
+	 Binaries\Win64\UnrealBuildAccelerator\*.exe ^
+	 | %FIND% ".exe" >> ..\Intermediate\Build\AutomationToolFiles.txt
+) 2>nul
+
 if not exist ..\Platforms goto NoPlatforms
 for /d %%D in (..\Platforms\*) do (
 	if exist %%D\Source\Programs\UnrealBuildTool (
@@ -73,6 +82,14 @@ for /d %%D in (..\Restricted\*) do (
 		 %%D\Source\Programs\UnrealBuildTool\*.csproj ^
 		 | %FIND% ".cs" >> ..\Intermediate\Build\AutomationToolFiles.txt
 	) 2>nul
+	if exist %%D\Binaries\Win64\UnrealBuildAccelerator (
+		dir /s ^
+		 %%D\Binaries\Win64\UnrealBuildAccelerator\*.dll ^
+		 | %FIND% ".dll" >> ..\Intermediate\Build\AutomationToolFiles.txt
+		dir /s ^
+		 %%D\Binaries\Win64\UnrealBuildAccelerator\*.exe ^
+		 | %FIND% ".exe" >> ..\Intermediate\Build\AutomationToolFiles.txt
+	) 2>nul
 )
 :NoRestricted
 
diff --git a/Engine/Build/BatchFiles/BuildUBT.bat b/Engine/Build/BatchFiles/BuildUBT.bat
--- a/Engine/Build/BatchFiles/BuildUBT.bat
+++ b/Engine/Build/BatchFiles/BuildUBT.bat
@@ -46,6 +46,15 @@ dir /s ^
  Programs\UnrealBuildTool\*.csproj ^
  | %FIND% ".cs" > ..\Intermediate\Build\UnrealBuildToolFiles.txt
 
+if exist Binaries\Win64\UnrealBuildAccelerator (
+	dir /s ^
+	 Binaries\Win64\UnrealBuildAccelerator\*.dll ^
+	 | %FIND% ".dll" >> ..\Intermediate\Build\UnrealBuildToolFiles.txt
+	dir /s ^
+	 Binaries\Win64\UnrealBuildAccelerator\*.exe ^
+	 | %FIND% ".exe" >> ..\Intermediate\Build\UnrealBuildToolFiles.txt
+) 2>nul
+
 if not exist ..\Platforms goto NoPlatforms
 for /d %%D in (..\Platforms\*) do (
 	if exist %%D\Source\Programs\UnrealBuildTool (
@@ -69,6 +78,14 @@ for /d %%D in (..\Restricted\*) do (
 		 %%D\Source\Programs\UnrealBuildTool\*.csproj ^
 		 | %FIND% ".cs" >> ..\Intermediate\Build\UnrealBuildToolFiles.txt
 	) 2>nul
+	if exist %%D\Binaries\Win64\UnrealBuildAccelerator (
+		dir /s ^
+		 %%D\Binaries\Win64\UnrealBuildAccelerator\*.dll ^
+		 | %FIND% ".dll" >> ..\Intermediate\Build\UnrealBuildToolFiles.txt
+		dir /s ^
+		 %%D\Binaries\Win64\UnrealBuildAccelerator\*.exe ^
+		 | %FIND% ".exe" >> ..\Intermediate\Build\UnrealBuildToolFiles.txt
+	) 2>nul
 )
 :NoRestricted
 
