UBA: Cleanup for work in progress linux support

diff --git a/Engine/Build/BatchFiles/BuildUAT.sh b/Engine/Build/BatchFiles/BuildUAT.sh
--- a/Engine/Build/BatchFiles/BuildUAT.sh
+++ b/Engine/Build/BatchFiles/BuildUAT.sh
@@ -74,9 +74,24 @@ elif [ -f ../Intermediate/Build/AutomationToolLastBuildTime ]; then
     echo "$UPDATED_AUTOMATIONTOOL_FILES"
   fi
 
+  UPDATED_DEP_FILES="$(find \
+    ../Binaries/Linux* \
+    ../Binaries/Mac \
+    ../Restricted/**/Binaries/Linux* \
+    ../Restricted/**/Binaries/Mac \
+    -maxdepth 2 \
+    -type f \
+    \( -iname \*Uba* \) \
+    -newer ../Intermediate/Build/UnrealBuildToolLastBuildTime)"
+  if [ -n "$UPDATED_DEP_FILES" ]; then
+    PERFORM_REBUILD=1
+    echo "Rebuilding: Found updated files:"
+    echo "$UPDATED_DEP_FILES"
+  fi
+
 else
   PERFORM_REBUILD=1
-  echo "Rebuilding: No record of previous builld"
+  echo "Rebuilding: No record of previous build"
 fi
 
 if [ $PERFORM_REBUILD -eq 1 ]; then
diff --git a/Engine/Build/BatchFiles/BuildUBT.sh b/Engine/Build/BatchFiles/BuildUBT.sh
--- a/Engine/Build/BatchFiles/BuildUBT.sh
+++ b/Engine/Build/BatchFiles/BuildUBT.sh
@@ -71,9 +71,24 @@ elif [ -f ../Intermediate/Build/UnrealBuildToolLastBuildTime ]; then
     echo "$UPDATED_DEP_FILES"
   fi
 
+  UPDATED_DEP_FILES="$(find \
+    ../Binaries/Linux* \
+    ../Binaries/Mac \
+    ../Restricted/**/Binaries/Linux* \
+    ../Restricted/**/Binaries/Mac \
+    -maxdepth 2 \
+    -type f \
+    \( -iname \*Uba* \) \
+    -newer ../Intermediate/Build/UnrealBuildToolLastBuildTime)"
+  if [ -n "$UPDATED_DEP_FILES" ]; then
+    PERFORM_REBUILD=1
+    echo "Rebuilding: Found updated files:"
+    echo "$UPDATED_DEP_FILES"
+  fi
+
 else
   PERFORM_REBUILD=1
-  echo "Rebuilding: No record of previous builld"
+  echo "Rebuilding: No record of previous build"
 fi
 
 if [ $PERFORM_REBUILD -eq 1 ]; then
