Fix path to UBA binaries in Build batch files so UBT wil be rebuilt

#rnx

diff --git a/Engine/Build/BatchFiles/BuildUAT.bat b/Engine/Build/BatchFiles/BuildUAT.bat
--- a/Engine/Build/BatchFiles/BuildUAT.bat
+++ b/Engine/Build/BatchFiles/BuildUAT.bat
@@ -46,12 +46,12 @@ dir /s ^
  Programs\UnrealBuildTool\*.csproj ^
  | %FIND% ".cs" > ..\Intermediate\Build\AutomationToolFiles.txt
 
-if exist Binaries\Win64\UnrealBuildAccelerator (
+if exist ..\Binaries\Win64\UnrealBuildAccelerator (
 	dir /s ^
-	 Binaries\Win64\UnrealBuildAccelerator\*.dll ^
+	 ..\Binaries\Win64\UnrealBuildAccelerator\*.dll ^
 	 | %FIND% ".dll" >> ..\Intermediate\Build\AutomationToolFiles.txt
 	dir /s ^
-	 Binaries\Win64\UnrealBuildAccelerator\*.exe ^
+	 ..\Binaries\Win64\UnrealBuildAccelerator\*.exe ^
 	 | %FIND% ".exe" >> ..\Intermediate\Build\AutomationToolFiles.txt
 ) 2>nul
 
diff --git a/Engine/Build/BatchFiles/BuildUBT.bat b/Engine/Build/BatchFiles/BuildUBT.bat
--- a/Engine/Build/BatchFiles/BuildUBT.bat
+++ b/Engine/Build/BatchFiles/BuildUBT.bat
@@ -46,12 +46,12 @@ dir /s ^
  Programs\UnrealBuildTool\*.csproj ^
  | %FIND% ".cs" > ..\Intermediate\Build\UnrealBuildToolFiles.txt
 
-if exist Binaries\Win64\UnrealBuildAccelerator (
+if exist ..\Binaries\Win64\UnrealBuildAccelerator (
 	dir /s ^
-	 Binaries\Win64\UnrealBuildAccelerator\*.dll ^
+	 ..\Binaries\Win64\UnrealBuildAccelerator\*.dll ^
 	 | %FIND% ".dll" >> ..\Intermediate\Build\UnrealBuildToolFiles.txt
 	dir /s ^
-	 Binaries\Win64\UnrealBuildAccelerator\*.exe ^
+	 ..\Binaries\Win64\UnrealBuildAccelerator\*.exe ^
 	 | %FIND% ".exe" >> ..\Intermediate\Build\UnrealBuildToolFiles.txt
 ) 2>nul
 
