Fix BuildUAT.sh checking for the wrong uptodate file

#rnx

diff --git a/Engine/Build/BatchFiles/BuildUAT.sh b/Engine/Build/BatchFiles/BuildUAT.sh
--- a/Engine/Build/BatchFiles/BuildUAT.sh
+++ b/Engine/Build/BatchFiles/BuildUAT.sh
@@ -80,7 +80,7 @@ elif [ -f ../Intermediate/Build/AutomationToolLastBuildTime ]; then
     -maxdepth 2 \
     -type f \
     \( -iname \*Uba* \) \
-    -newer ../Intermediate/Build/UnrealBuildToolLastBuildTime)"
+    -newer ../Intermediate/Build/AutomationToolLastBuildTime)"
   if [ -n "$UPDATED_DEP_FILES" ]; then
     PERFORM_REBUILD=1
     echo "Rebuilding: Found updated files:"
