Adjust error message when calling RunUBT.bat as errors will already be printed to the console

#rnx

diff --git a/Engine/Build/BatchFiles/RunUBT.bat b/Engine/Build/BatchFiles/RunUBT.bat
--- a/Engine/Build/BatchFiles/RunUBT.bat
+++ b/Engine/Build/BatchFiles/RunUBT.bat
@@ -65,7 +65,7 @@ goto Exit
 
 :Error_UBTRunFailed
 echo.
-echo RUNUBT ERROR: UnrealBuildTool exited with a failure
+echo RunUBT ERROR: UnrealBuildTool exited with a failure
 echo.
 set RUNUBT_EXITCODE=%ERRORLEVEL%
 goto Exit
