Fix RunUBT.sh failures in Installed Builds

diff --git a/Engine/Build/BatchFiles/RunUBT.sh b/Engine/Build/BatchFiles/RunUBT.sh
--- a/Engine/Build/BatchFiles/RunUBT.sh
+++ b/Engine/Build/BatchFiles/RunUBT.sh
@@ -26,7 +26,7 @@ if [ "$(uname)" = "Linux" ]; then
 	source "$BASE_PATH/Linux/SetupEnvironment.sh" $EnvironmentType "$BASE_PATH/Linux"
 fi
 
-if [[ "$*" != *-SkipUBTBuild* ]]; then
+if [[ "$*" != *-SkipUBTBuild* && ! -f "$BASE_PATH/../InstalledBuild.txt" ]]; then
   if [ -f "$BASE_PATH/../../Source/Programs/UnrealBuildTool/UnrealBuildTool.csproj" ]; then 
     echo "Building UBT..."
     "$BASE_PATH/BuildUBT.sh"
