Remove references to mono from some Linux platform files (misleading comments)

diff --git a/Engine/Build/BatchFiles/Linux/README.md b/Engine/Build/BatchFiles/Linux/README.md
--- a/Engine/Build/BatchFiles/Linux/README.md
+++ b/Engine/Build/BatchFiles/Linux/README.md
@@ -21,7 +21,6 @@ and an up-to-date list should be maintained (and installed) by Setup.sh -
 feel free to suggest modifications. Automated install generally works for Ubuntu only.
 
 Most important dependencies:
-- mono 3.x (2.x may work, but is not recommended), including xbuild and C# compiler (*mcs), and libraries for NET 4.0 framework.
 - clang 3.9.0 (clang 3.5 through 3.8 should also be able to compile the engine).
 
 You will also need at least 20 GB of free disk space and a relatively powerful
@@ -71,8 +70,7 @@ How to set up the sources for building, step by step:
 
 Updating the sources later can be done with git pull. The tool to download binary files will be
 registered as a post-merge hook by Setup.sh, so third party libraries will be updated automatically
-(if needed). If you ever need to run it directly, it can be found in Engine/Binaries/DotNET/ directory
-(GitDependencies.exe, which needs to be invoked through mono).
+(if needed). If you ever need to run it directly, it can be found in Engine/Binaries/DotNET/ directory.
 
 
 
@@ -121,7 +119,7 @@ You can also append -game if you want to run the project as a game (you can also
 Notes
 -----
 
-Depending on mono version and some other not yet clarified circumstances (this may be relevant: 
+Depending on dotnet version and some other not yet clarified circumstances (this may be relevant: 
 http://stackoverflow.com/questions/13859467/ravendb-client-onlinux-connecting-to-windows-server-using-mono-http),
 binary downloader tool invoked by Setup.sh may fail. In that case, Setup.sh will keep re-running it
 until it succeeds (or at least stops crashing). You may want to keep an eye on this as there is
diff --git a/Engine/Build/BatchFiles/RunUAT.sh b/Engine/Build/BatchFiles/RunUAT.sh
--- a/Engine/Build/BatchFiles/RunUAT.sh
+++ b/Engine/Build/BatchFiles/RunUAT.sh
@@ -114,7 +114,7 @@ else
 fi
 
 # if we are running under UE, we need to run this with the term handler (otherwise canceling a UAT job from the editor
-# can leave mono, etc running in the background, which means we need the PID so we 
+# can leave dotnet, etc running in the background, which means we need the PID so we 
 # run it in the background
 if [ "$UE_DesktopUnrealProcess" = "1" ]; then
 	# you can't set a dotted env var nicely in sh, but env will run a command with
