Fix batch up to date check erroring if a directory is missing

#rnx

diff --git a/Engine/Build/BatchFiles/BuildUAT.bat b/Engine/Build/BatchFiles/BuildUAT.bat
--- a/Engine/Build/BatchFiles/BuildUAT.bat
+++ b/Engine/Build/BatchFiles/BuildUAT.bat
@@ -57,31 +57,27 @@ if exist ..\Binaries\Win64\UnrealBuildAccelerator (
 
 if not exist ..\Platforms goto NoPlatforms
 for /d %%D in (..\Platforms\*) do (
-	if exist %%D\Source\Programs\UnrealBuildTool (
-		dir /s ^
-		 %%D\Source\Programs\AutomationTool\*.cs ^
-		 %%D\Source\Programs\AutomationTool\*.csproj ^
-		 %%D\Source\Programs\Shared\*.cs ^
-		 %%D\Source\Programs\Shared\*.csproj ^
-		 %%D\Source\Programs\UnrealBuildTool\*.cs ^
-		 %%D\Source\Programs\UnrealBuildTool\*.csproj ^
-		 | %FIND% ".cs" >> ..\Intermediate\Build\AutomationToolFiles.txt
-	) 2>nul
+	for %%F in (AutomationTool Shared UnrealBuildTool) do (
+		if exist %%D\Source\Programs\%%F (
+			dir /s ^
+			%%D\Source\Programs\%%F\*.cs ^
+			%%D\Source\Programs\%%F\*.csproj ^
+			| %FIND% ".cs" >> ..\Intermediate\Build\AutomationToolFiles.txt
+		) 2>nul
+	)
 )
 :NoPlatforms
 
 if not exist ..\Restricted goto NoRestricted
 for /d %%D in (..\Restricted\*) do (
-	if exist %%D\Source\Programs\UnrealBuildTool (
-		dir /s ^
-		 %%D\Source\Programs\AutomationTool\*.cs ^
-		 %%D\Source\Programs\AutomationTool\*.csproj ^
-		 %%D\Source\Programs\Shared\*.cs ^
-		 %%D\Source\Programs\Shared\*.csproj ^
-		 %%D\Source\Programs\UnrealBuildTool\*.cs ^
-		 %%D\Source\Programs\UnrealBuildTool\*.csproj ^
-		 | %FIND% ".cs" >> ..\Intermediate\Build\AutomationToolFiles.txt
-	) 2>nul
+	for %%F in (AutomationTool Shared UnrealBuildTool) do (
+		if exist %%D\Source\Programs\%%F (
+			dir /s ^
+			%%D\Source\Programs\%%F\*.cs ^
+			%%D\Source\Programs\%%F\*.csproj ^
+			| %FIND% ".cs" >> ..\Intermediate\Build\AutomationToolFiles.txt
+		) 2>nul
+	)
 )
 :NoRestricted
 
diff --git a/Engine/Build/BatchFiles/BuildUAT.sh b/Engine/Build/BatchFiles/BuildUAT.sh
--- a/Engine/Build/BatchFiles/BuildUAT.sh
+++ b/Engine/Build/BatchFiles/BuildUAT.sh
@@ -47,6 +47,8 @@ elif [ -f ../Intermediate/Build/AutomationToolLastBuildTime ]; then
     Programs/Shared/EpicGames.UBA \
     Programs/Shared/EpicGames.UHT \
     Programs/UnrealBuildTool \
+    ../Restricted/**/Source/Programs/Shared \
+    ../Platforms/*/Source/Programs/Shared \
     ../Restricted/**/Source/Programs/UnrealBuildTool \
     ../Platforms/*/Source/Programs/UnrealBuildTool \
     ../Restricted/**/Source/Programs/AutomationTool \
diff --git a/Engine/Build/BatchFiles/BuildUBT.bat b/Engine/Build/BatchFiles/BuildUBT.bat
--- a/Engine/Build/BatchFiles/BuildUBT.bat
+++ b/Engine/Build/BatchFiles/BuildUBT.bat
@@ -57,27 +57,27 @@ if exist ..\Binaries\Win64\UnrealBuildAccelerator (
 
 if not exist ..\Platforms goto NoPlatforms
 for /d %%D in (..\Platforms\*) do (
-	if exist %%D\Source\Programs\UnrealBuildTool (
-		dir /s ^
-		 %%D\Source\Programs\Shared\*.cs ^
-		 %%D\Source\Programs\Shared\*.csproj ^
-		 %%D\Source\Programs\UnrealBuildTool\*.cs ^
-		 %%D\Source\Programs\UnrealBuildTool\*.csproj ^
-		 | %FIND% ".cs" >> ..\Intermediate\Build\UnrealBuildToolFiles.txt
-	) 2>nul
+	for %%F in (Shared UnrealBuildTool) do (
+		if exist %%D\Source\Programs\%%F (
+			dir /s ^
+			%%D\Source\Programs\%%F\*.cs ^
+			%%D\Source\Programs\%%F\*.csproj ^
+			| %FIND% ".cs" >> ..\Intermediate\Build\UnrealBuildToolFiles.txt
+		) 2>nul
+	)
 )
 :NoPlatforms
 
 if not exist ..\Restricted goto NoRestricted
 for /d %%D in (..\Restricted\*) do (
-	if exist %%D\Source\Programs\UnrealBuildTool (
-		dir /s ^
-		 %%D\Source\Programs\Shared\*.cs ^
-		 %%D\Source\Programs\Shared\*.csproj ^
-		 %%D\Source\Programs\UnrealBuildTool\*.cs ^
-		 %%D\Source\Programs\UnrealBuildTool\*.csproj ^
-		 | %FIND% ".cs" >> ..\Intermediate\Build\UnrealBuildToolFiles.txt
-	) 2>nul
+	for %%F in (Shared UnrealBuildTool) do (
+		if exist %%D\Source\Programs\%%F (
+			dir /s ^
+			%%D\Source\Programs\%%F\*.cs ^
+			%%D\Source\Programs\%%F\*.csproj ^
+			| %FIND% ".cs" >> ..\Intermediate\Build\UnrealBuildToolFiles.txt
+		) 2>nul
+	)
 )
 :NoRestricted
 
diff --git a/Engine/Build/BatchFiles/BuildUBT.sh b/Engine/Build/BatchFiles/BuildUBT.sh
--- a/Engine/Build/BatchFiles/BuildUBT.sh
+++ b/Engine/Build/BatchFiles/BuildUBT.sh
@@ -47,6 +47,8 @@ elif [ -f ../Intermediate/Build/UnrealBuildToolLastBuildTime ]; then
     Programs/Shared/EpicGames.UBA \
     Programs/Shared/EpicGames.UHT \
     Programs/UnrealBuildTool \
+    ../Restricted/**/Source/Programs/Shared \
+    ../Platforms/*/Source/Programs/Shared \
     ../Restricted/**/Source/Programs/UnrealBuildTool \
     ../Platforms/*/Source/Programs/UnrealBuildTool \
     -type f \
