- Fixed bad version generation in XCode script

diff --git a/Engine/Build/BatchFiles/Mac/UpdateVersionAfterBuild.sh b/Engine/Build/BatchFiles/Mac/UpdateVersionAfterBuild.sh
--- a/Engine/Build/BatchFiles/Mac/UpdateVersionAfterBuild.sh
+++ b/Engine/Build/BatchFiles/Mac/UpdateVersionAfterBuild.sh
@@ -14,7 +14,7 @@ else
 	CL=$3
 fi
 
-VERSION="$CL.1"
+VERSION="0.1"
 # increment version in the counter file
 if [ -f "$VERSION_FILE" ]; then
 	VERSION=$(cat "$VERSION_FILE")
@@ -57,7 +57,7 @@ XCCONFIG_FILE="$1/Intermediate/Build/Versions.xcconfig"
 
 mkdir -p "$1/Intermediate/Build"
 echo "UE_MAC_BUILD_VERSION = $CL.$MAC_VERSION" > "$XCCONFIG_FILE"
-echo "UE_IOS_BUILD_VERSION = $CL.$IOS_VERSION.$CL" >> "$XCCONFIG_FILE"
-echo "UE_TVOS_BUILD_VERSION = $CL.$TVOS_VERSION.$CL" >> "$XCCONFIG_FILE"
-echo "UE_VISIONOS_BUILD_VERSION = $CL.$VISIONOS_VERSION.$CL" >> "$XCCONFIG_FILE"
+echo "UE_IOS_BUILD_VERSION = $CL.$IOS_VERSION" >> "$XCCONFIG_FILE"
+echo "UE_TVOS_BUILD_VERSION = $CL.$TVOS_VERSION" >> "$XCCONFIG_FILE"
+echo "UE_VISIONOS_BUILD_VERSION = $CL.$VISIONOS_VERSION" >> "$XCCONFIG_FILE"
 
