- Added a ResignApp.sh (meant for IOS but should work for Mac apps also) that can resign a .app or .ipa with another cert/provision pair, and change bundle ID and commandline

diff --git a/Engine/Build/BatchFiles/Mac/ResignApp.sh b/Engine/Build/BatchFiles/Mac/ResignApp.sh
new file mode 100755
--- /dev/null
+++ b/Engine/Build/BatchFiles/Mac/ResignApp.sh
@@ -0,0 +1,149 @@
+# !/bin/bash
+
+set -e
+
+DoWork()
+{
+	echo "------------------------------------------------------------------------------"
+	echo "Resigning with certificate: $DEVELOPER"
+
+	rm -rf working
+	if [[ $SOURCEAPP == *.ipa ]]; then
+		echo "Prep: Unzipping $SOURCEAPP to 'working'..."
+		unzip -qo "$SOURCEAPP" -d working
+		TARGET="working/Payload/$(ls working/Payload/)"
+	else
+		if [[ $TARGETAPP == *.ipa ]]; then
+			mkdir -p working/Payload/
+			TARGET="working/Payload/$(basename $SOURCEAPP)"
+		else
+			TARGET="$TARGETAPP"
+		fi
+		echo "Prep: Copying $SOURCEAPP to $TARGET..."
+		ditto "$SOURCEAPP" "$TARGET"
+	fi
+
+	if [[ ! -z "$MOBILEPROV" ]]; then
+		echo "Prep: Copying $SOURCEAPP into app..."
+		cp "$MOBILEPROV" "$TARGET/embedded.mobileprovision"
+	fi
+
+	find -d "$TARGET"  \( -name "*.app" -o -name "*.appex" -o -name "*.framework" -o -name "*.dylib" \) > directories.txt
+	oldbundleid=$(/usr/libexec/PlistBuddy -c "Print:CFBundleIdentifier" "$TARGET/Info.plist")
+
+	if [[ ! -z "$BUNDLE" ]]; then
+	   echo "Changing BundleID from $oldbundleid with : $BUNDLE"
+	   /usr/libexec/PlistBuddy -c "Set:CFBundleIdentifier $BUNDLE" "$TARGET/Info.plist"
+	fi
+
+	if [[ ! -z "$CMDLINE" ]]; then
+		echo "Setting commandline to $CMDLINE"
+		echo $CMDLINE > "$TARGET/uecommandline.txt"
+	fi
+
+	echo "------------------------------------------------------------------------------"
+
+	while IFS='' read -r line || [[ -n "$line" ]]; do
+
+		if [[ ! -z "$BUNDLE" ]] && [[ "$line" == *".appex"* ]]; then
+		   extbundleid=$(/usr/libexec/PlistBuddy -c "Print:CFBundleIdentifier" "$line/Info.plist")
+		   echo "Changing .appex BundleID $extbundleid with : ${extbundleid/$oldbundleid/$BUNDLE}"
+		   /usr/libexec/PlistBuddy -c "Set:CFBundleIdentifier ${extbundleid/$oldbundleid/$BUNDLE}" "$line/Info.plist"
+		fi    
+
+		echo ""
+		echo Codesigning $line [/usr/bin/codesign --preserve-metadata=entitlements,flags,identifier --continue -f -s "$DEVELOPER" "$line"]
+		/usr/bin/codesign --preserve-metadata=entitlements,flags,identifier --continue -f -s "$DEVELOPER" "$line"
+
+	done < directories.txt
+
+	if [[ $TARGETAPP == *.ipa ]]; then
+		echo ""
+		echo Zipping working to $TARGETAPP...
+		cd working
+		zip -qry ../working.ipa *
+		cd ..
+		mv working.ipa "$TARGETAPP"
+		echo Cleanup up 'working' dir...
+		rm -rf working
+	elif [[ $SOURCEAPP == *.ipa ]]; then
+		echo Moving $TARGET to $TARGETAPP...
+		mv "$TARGET" "$TARGETAPP"
+	fi
+}
+
+Help()
+{
+	echo ""
+	echo "Usage: ResignApp <options>"
+	echo ""
+	echo "REQUIRED OPTIONS:"
+	echo "  -s | --sourceapp <path to input>"
+	echo "     Path to the .app or .ipa you want to resign"
+	echo "  -d | --destapp <path to output>"
+	echo "     Path to the output .app or .ipa (note that input type and output type do not need to match)"
+	echo "  -i | --identity <signing identitiy> "
+	echo "     Name of the identity to sign with (this is the full name or prefix of a certificate in your login keychain - 'Apple Development' will often work)"
+	echo ""
+	echo "OPTIONAL OPTIONS:"
+	echo "  -p | --provision <path to .mobileprovision>"
+	echo "     Path to the .mobileprovision file to sign with"
+	echo "  -b | --bundleid"
+	echo "     Bundle ID to use in the app, repleacing existing bundle ID"
+	echo "  -c | --cmdline"
+	echo "     A commandline to place into the app as uecommandline.txt"
+	echo "  -h | --help"
+	echo "     Show this message"
+	exit 0
+}
+
+
+if [[ $# -eq 0 ]]; then
+	Help
+	exit 0
+fi
+
+while [[ $# -gt 0 ]]; do
+  case $1 in
+	-s|--sourceapp)
+	  SOURCEAPP="$2"
+	  shift # past argument
+	  shift # past value
+	  ;;
+	-i|--identity)
+	  DEVELOPER="$2"
+	  shift # past argument
+	  shift # past value
+	  ;;
+	-p|--provision)
+	  MOBILEPROV="$2"
+	  shift # past argument
+	  shift # past value
+	  ;;
+	-d|--destapp)
+	  TARGETAPP="$2"
+	  shift # past argument
+	  shift # past value
+	  ;;
+	-b|--bundleid)
+	  BUNDLE="$2"
+	  shift # past argument
+	  shift # past value
+	  ;;
+	-c|--cmdline)
+	  CMDLINE="$2"
+	  shift # past argument
+	  shift # past value
+	  ;;
+	-h|--help)
+	  Help
+	  shift # past value
+	  ;;
+	-*|--*)
+	  echo "Unknown option $1"
+	  exit 1
+	  ;;
+  esac
+done
+
+DoWork
