BAT: Updating UE4 text to UE5

diff --git a/Engine/Build/BatchFiles/MakeAndInstallSSHKey.bat b/Engine/Build/BatchFiles/MakeAndInstallSSHKey.bat
--- a/Engine/Build/BatchFiles/MakeAndInstallSSHKey.bat
+++ b/Engine/Build/BatchFiles/MakeAndInstallSSHKey.bat
@@ -57,7 +57,7 @@ popd
 
 echo.
 echo ===================================================================================================
-echo Assuming you have seen no errors above, you are ready to compile UE4 remotely on %MACHINE%!
+echo Assuming you have seen no errors above, you are ready to compile UE5 remotely on %MACHINE%!
 echo ===================================================================================================
 echo.
 
