- Updated ResignApp.sh
  - Properly fixes entitlements when changing bundle ID
  - Can now handle .xcarchives (a .xcarchive can resign to any type, .app and .ipa cannot output to .xcarchive, because we re-use the directory structure of the input .xcarchive)
  - Better handle spaces in paths

diff --git a/Engine/Build/BatchFiles/Mac/ResignApp.sh b/Engine/Build/BatchFiles/Mac/ResignApp.sh
--- a/Engine/Build/BatchFiles/Mac/ResignApp.sh
+++ b/Engine/Build/BatchFiles/Mac/ResignApp.sh
@@ -8,19 +8,34 @@ DoWork()
 	echo "Resigning with certificate: $DEVELOPER"
 
 	rm -rf working
-	if [[ $SOURCEAPP == *.ipa ]]; then
+	if [[ "$SOURCEAPP" == *.ipa ]]; then
 		echo "Prep: Unzipping $SOURCEAPP to 'working'..."
 		unzip -qo "$SOURCEAPP" -d working
 		TARGET="working/Payload/$(ls working/Payload/)"
 	else
-		if [[ $TARGETAPP == *.ipa ]]; then
-			mkdir -p working/Payload/
-			TARGET="working/Payload/$(basename $SOURCEAPP)"
-		else
-			TARGET="$TARGETAPP"
+		if [[ "$TARGETAPP" == *.xcarchive ]]; then
+			if [[ "$SOURCEAPP" != *.xcarchive ]]; then
+				echo "When destination is a .xcarchive, the source must also be a .xcarchive"
+				exit 1
+			fi
+			TARGET="$TARGETAPP/Products/Applications/$(ls $SOURCEAPP/Products/Applications/)"
+			echo "Prep: Copying $SOURCEAPP to $TARGETAPP..."
+			rm -rf "$TARGETAPP"
+			ditto "$SOURCEAPP" "$TARGETAPP"
+		else		
+			if [[ "$SOURCEAPP" == *.xcarchive ]]; then
+				SOURCEAPP="$SOURCEAPP/Products/Applications/$(ls $SOURCEAPP/Products/Applications/)"
+			fi
+			if [[ "$TARGETAPP" == *.ipa ]]; then
+				mkdir -p working/Payload/
+				TARGET="working/Payload/$(basename $SOURCEAPP)"
+			else
+				TARGET="$TARGETAPP"
+			fi
+			echo "Prep: Copying $SOURCEAPP to $TARGET..."
+			rm -rf "$TARGET"
+			ditto "$SOURCEAPP" "$TARGET"
 		fi
-		echo "Prep: Copying $SOURCEAPP to $TARGET..."
-		ditto "$SOURCEAPP" "$TARGET"
 	fi
 
 	if [[ ! -z "$MOBILEPROV" ]]; then
@@ -43,6 +58,7 @@ DoWork()
 
 	echo "------------------------------------------------------------------------------"
 
+	EXTRAOPTIONS=""
 	while IFS='' read -r line || [[ -n "$line" ]]; do
 
 		if [[ ! -z "$BUNDLE" ]] && [[ "$line" == *".appex"* ]]; then
@@ -51,13 +67,24 @@ DoWork()
 		   /usr/libexec/PlistBuddy -c "Set:CFBundleIdentifier ${extbundleid/$oldbundleid/$BUNDLE}" "$line/Info.plist"
 		fi    
 
+		if [[ -f "$line/embedded.mobileprovision" ]]; then
+			security cms -D -i "$line/embedded.mobileprovision" > t_entitlements_full.plist
+			/usr/libexec/PlistBuddy -x -c 'Print:Entitlements' t_entitlements_full.plist > t_entitlements.plist
+			TEAMID=$(/usr/libexec/PlistBuddy -c 'Print:com.apple.developer.team-identifier' t_entitlements.plist)
+			/usr/libexec/PlistBuddy -c "Set:application-identifier $TEAMID.$BUNDLE" t_entitlements.plist
+			
+			EXTRAOPTIONS="--entitlements t_entitlements.plist"
+		else
+			EXTRAOPTIONS="--preserve-metadata=entitlements,flags"
+		fi
+		
 		echo ""
-		echo Codesigning $line [/usr/bin/codesign --preserve-metadata=entitlements,flags,identifier --continue -f -s "$DEVELOPER" "$line"]
-		/usr/bin/codesign --preserve-metadata=entitlements,flags,identifier --continue -f -s "$DEVELOPER" "$line"
+		echo Codesigning $line [/usr/bin/codesign  --continue -f -s "$DEVELOPER" $EXTRAOPTIONS "$line"]
+		/usr/bin/codesign --continue -f -s "$DEVELOPER" $EXTRAOPTIONS "$line"
 
 	done < directories.txt
 
-	if [[ $TARGETAPP == *.ipa ]]; then
+	if [[ "$TARGETAPP" == *.ipa ]]; then
 		echo ""
 		echo Zipping working to $TARGETAPP...
 		cd working
@@ -66,10 +93,26 @@ DoWork()
 		mv working.ipa "$TARGETAPP"
 		echo Cleanup up 'working' dir...
 		rm -rf working
-	elif [[ $SOURCEAPP == *.ipa ]]; then
+	elif [[ "$TARGETAPP" == *.xcarchive ]]; then
+		echo ""
+		echo Updating $TARGETAPP/Info.plist to match any new settings...
+		if [[ ! -z "$BUNDLE" ]]; then
+			echo "  Setting CFBundleIdentifier to $BUNDLE"
+			/usr/libexec/PlistBuddy -c "Set:ApplicationProperties:CFBundleIdentifier $BUNDLE" "$TARGETAPP/Info.plist"
+		fi
+		echo "  Setting Team to $TEAMID"
+		/usr/libexec/PlistBuddy -c "Set:ApplicationProperties:Team $TEAMID" "$TARGETAPP/Info.plist"
+		FULLIDENTITY=$(security find-identity -v -p codesigning | grep "$DEVELOPER" | sed -n 's/.*\"\(.*\)\".*/\1/p;q')
+		echo "  Setting SigningIdentity to $FULLIDENTITY"
+		/usr/libexec/PlistBuddy -c "Set:ApplicationProperties:SigningIdentity $FULLIDENTITY" "$TARGETAPP/Info.plist"
+	elif [[ "$SOURCEAPP" == *.ipa ]]; then
 		echo Moving $TARGET to $TARGETAPP...
 		mv "$TARGET" "$TARGETAPP"
 	fi
+	
+	rm -f directories.txt
+	rm -f t_entitlements_full.plist
+	rm -f t_entitlements.plist
 }
 
 Help()
