- Added support ro ResignApp.sh for setting the marketing version (like 1.0) and bundle version (0.1.123123412)
- Couple of fixes when not changing bundle ID

diff --git a/Engine/Build/BatchFiles/Mac/ResignApp.sh b/Engine/Build/BatchFiles/Mac/ResignApp.sh
--- a/Engine/Build/BatchFiles/Mac/ResignApp.sh
+++ b/Engine/Build/BatchFiles/Mac/ResignApp.sh
@@ -47,8 +47,20 @@ DoWork()
 	oldbundleid=$(/usr/libexec/PlistBuddy -c "Print:CFBundleIdentifier" "$TARGET/Info.plist")
 
 	if [[ ! -z "$BUNDLE" ]]; then
-	   echo "Changing BundleID from $oldbundleid with : $BUNDLE"
-	   /usr/libexec/PlistBuddy -c "Set:CFBundleIdentifier $BUNDLE" "$TARGET/Info.plist"
+		echo "Changing BundleID from $oldbundleid with : $BUNDLE"	
+		/usr/libexec/PlistBuddy -c "Set:CFBundleIdentifier $BUNDLE" "$TARGET/Info.plist"
+	fi
+
+	if [[ ! -z "$MARKETING_VERSION" ]]; then
+		oldmarketingversion=$(/usr/libexec/PlistBuddy -c "Print:CFBundleShortVersionString" "$TARGET/Info.plist")	
+		echo "Changing Marketing Version from $oldmarketingversion to $MARKETING_VERSION"
+		/usr/libexec/PlistBuddy -c "Set:CFBundleShortVersionString $BUNDLE" "$TARGET/Info.plist"
+	fi
+
+	if [[ ! -z "$BUNDLE_VERSION" ]]; then
+		oldbundleversion=$(/usr/libexec/PlistBuddy -c "Print:CFBundleVersion" "$TARGET/Info.plist")	
+		echo "Changing Bundle Version from $oldbundleversion to $BUNDLE_VERSION"
+		/usr/libexec/PlistBuddy -c "Set:CFBundleVersion $BUNDLE_VERSION" "$TARGET/Info.plist"
 	fi
 
 	if [[ ! -z "$CMDLINE" ]]; then
@@ -73,9 +85,11 @@ DoWork()
 			TEAMID=$(/usr/libexec/PlistBuddy -c 'Print:com.apple.developer.team-identifier' t_entitlements.plist)
 			/usr/libexec/PlistBuddy -c "Set:application-identifier $TEAMID.$BUNDLE" t_entitlements.plist
 			
-			EXTRAOPTIONS="--entitlements t_entitlements.plist"
-		else
-			EXTRAOPTIONS="--preserve-metadata=entitlements,flags"
+			if [[ ! -z "$BUNDLE" ]]; then
+				EXTRAOPTIONS="--entitlements t_entitlements.plist"
+			else
+				EXTRAOPTIONS="--preserve-metadata=entitlements,flags,identifier"
+			fi
 		fi
 		
 		echo ""
@@ -100,11 +114,23 @@ DoWork()
 			echo "  Setting CFBundleIdentifier to $BUNDLE"
 			/usr/libexec/PlistBuddy -c "Set:ApplicationProperties:CFBundleIdentifier $BUNDLE" "$TARGETAPP/Info.plist"
 		fi
+
 		echo "  Setting Team to $TEAMID"
 		/usr/libexec/PlistBuddy -c "Set:ApplicationProperties:Team $TEAMID" "$TARGETAPP/Info.plist"
 		FULLIDENTITY=$(security find-identity -v -p codesigning | grep "$DEVELOPER" | sed -n 's/.*\"\(.*\)\".*/\1/p;q')
 		echo "  Setting SigningIdentity to $FULLIDENTITY"
 		/usr/libexec/PlistBuddy -c "Set:ApplicationProperties:SigningIdentity $FULLIDENTITY" "$TARGETAPP/Info.plist"
+
+		if [[ ! -z "$MARKETING_VERSION" ]]; then
+		   echo "  Changing Marketing Version to $MARKETING_VERSION"
+		   /usr/libexec/PlistBuddy -c "Set:ApplicationProperties:CFBundleShortVersionString $MARKETING_VERSION" "$TARGETAPP/Info.plist"
+		fi
+		
+		if [[ ! -z "$BUNDLE_VERSION" ]]; then
+			echo "  Changing Bundle Version to $BUNDLE_VERSION"
+			/usr/libexec/PlistBuddy -c "Set:ApplicationProperties:CFBundleVersion $BUNDLE_VERSION" "$TARGETAPP/Info.plist"
+		fi
+
 	elif [[ "$SOURCEAPP" == *.ipa ]]; then
 		echo Moving $TARGET to $TARGETAPP...
 		mv "$TARGET" "$TARGETAPP"
@@ -133,6 +159,10 @@ Help()
 	echo "     Path to the .mobileprovision file to sign with"
 	echo "  -b | --bundleid"
 	echo "     Bundle ID to use in the app, repleacing existing bundle ID"
+	echo "  -mv | --marketingversion"
+	echo "     Marketing version for the app. This is the version TestFlight/AppStoreConnect uses to manage builds"
+	echo "  -bv | --bundleversion"
+	echo "     Bundle version for the app. This is the internal version string that are gathered within one Marketing Version in TF/ASC"
 	echo "  -c | --cmdline"
 	echo "     A commandline to place into the app as uecommandline.txt"
 	echo "  -h | --help"
@@ -173,6 +203,16 @@ while [[ $# -gt 0 ]]; do
 	  shift # past argument
 	  shift # past value
 	  ;;
+	-mv|--marketingversion)
+	  MARKETING_VERSION="$2"
+	  shift # past argument
+	  shift # past value
+	  ;;
+	-bv|--bundleversion)
+	  BUNDLE_VERSION="$2"
+	  shift # past argument
+	  shift # past value
+	  ;;
 	-c|--cmdline)
 	  CMDLINE="$2"
 	  shift # past argument
