Toolchain upgrade to llvm 18.1.0

fix scripts
    - support Rocky Linux 8 as minimum supported OS
    - deprecate CentOS 7 support
    - update to using a vanilla released version of crosstool
        - configs should now be able to be upgraded when we go to a new version instead of having to be rebuilt
    - changes to script flow to make it easier to iterate on the build process when things break
        - stages of the build are checkpointed so that they don't have to be redone each run

diff --git a/Engine/Build/BatchFiles/Linux/Toolchain/RunMe.bat b/Engine/Build/BatchFiles/Linux/Toolchain/RunMe.bat
--- a/Engine/Build/BatchFiles/Linux/Toolchain/RunMe.bat
+++ b/Engine/Build/BatchFiles/Linux/Toolchain/RunMe.bat
@@ -1,22 +1,24 @@
 @REM @echo off
 
-set TOOLCHAIN_VERSION=v22
-set LLVM_VERSION=16.0.6
-set LLVM_BRANCH=release/16.x
-set LLVM_TAG=llvmorg-16.0.6
+set TOOLCHAIN_VERSION=v23
+set LLVM_VERSION=18.1.0
+set LLVM_BRANCH=release/18.x
+set LLVM_TAG=llvmorg-18.1.0
 
 set CMAKE_BINARY=%CD%\..\..\..\..\Extras\ThirdPartyNotUE\CMake\bin\cmake.exe
 set PYTHON_BINARY=%CD%\..\..\..\..\Binaries\ThirdParty\Python\Win64\python.exe
-set ZLIB_PATH=%CD%\..\..\..\..\..\Engine\Source\ThirdParty\zlib\v1.2.8
+set ZLIB_PATH=%CD%\..\..\..\..\..\Engine\Source\ThirdParty\zlib\1.3
 set "ZLIB_PATH=%ZLIB_PATH:\=/%"
 set NSIS_BINARY=C:\Program Files (x86)\NSIS\Bin\makensis.exe
 set PATCH_BINARY=C:\Program Files\Git\usr\bin\patch.exe
 set VS_VERSION="Visual Studio 17 2022"
 
-for %%i in (python.exe) do set PYTHON_BINARY="%%~$PATH:i"
-for %%i in (cmake.exe) do set CMAKE_BINARY="%%~$PATH:i"
+rem Use the following two lines if you want to override python or cmake with
+rem the version in your path
+rem for %%i in (python.exe) do set PYTHON_BINARY="%%~$PATH:i"
+rem for %%i in (cmake.exe) do set CMAKE_BINARY="%%~$PATH:i"
 
-set FILENAME=%TOOLCHAIN_VERSION%_clang-%LLVM_VERSION%-centos7
+set FILENAME=%TOOLCHAIN_VERSION%_clang-%LLVM_VERSION%-rockylinux8
 
 echo Building %FILENAME%.exe...
 
@@ -28,15 +30,19 @@ echo Using VisualStudio: %VS_VERSION%
 @REM We need to build in a directory with shorter path, so we avoid hitting path max limit.
 set ROOT_DIR=%CD%
 
-rmdir /S /Q %TEMP%\clang-build-%LLVM_VERSION%
+rmdir /S /Q %TEMP%\clang-build-%LLVM_VERSION% > nul
 mkdir %TEMP%\clang-build-%LLVM_VERSION%
 pushd %TEMP%\clang-build-%LLVM_VERSION%
 
 
-unzip -q -o %ROOT_DIR%\%FILENAME%-windows.zip -d OUTPUT
+rem The MS provided tar can be used to create and extract zip files on windows
+mkdir OUTPUT
+pushd OUTPUT
+tar xf %ROOT_DIR%\%FILENAME%-windows.zip
+popd
 
 echo Cloning rpmalloc to speed up clang/lld runtime operation
-git clone [REDACTED_PROFILE_URL]/rpmalloc rpmalloc --depth 1
+git clone [REDACTED_PROFILE_URL]/rpmalloc -b main rpmalloc --depth 1
 
 
 echo Cloning LLVM (tag %LLVM_TAG% only)
@@ -50,17 +56,20 @@ echo Applying patches
 pushd "source"
 rem set DRY_RUN=--dry-run
 set DRY_RUN=
-"%PATCH_BINARY%" %DRY_RUN% -p 1 -i %ROOT_DIR%\patches\clang\default-dwarf-4.patch 
-"%PATCH_BINARY%" %DRY_RUN% -p 1 -i %ROOT_DIR%\patches\compiler-rt\manually-define-AT_HWCAP2.diff 
-"%PATCH_BINARY%" %DRY_RUN% -p 1 -i %ROOT_DIR%\patches\llvm\disable-auto-upgrade-debug-info.patch 
+rem disabled but left in place in case testing reveals that they are still needed
+rem "%PATCH_BINARY%" %DRY_RUN% -p 1 -i %ROOT_DIR%\patches\clang\default-dwarf-4.patch 
+rem "%PATCH_BINARY%" %DRY_RUN% -p 1 -i %ROOT_DIR%\patches\compiler-rt\manually-define-AT_HWCAP2.diff 
+rem "%PATCH_BINARY%" %DRY_RUN% -p 1 -i %ROOT_DIR%\patches\llvm\disable-auto-upgrade-debug-info.patch 
 popd
 
 
 echo Building LLVM, clang, lld, bolt
 mkdir build_all
 pushd build_all
-%CMAKE_BINARY% -G %VS_VERSION% -DLLVM_ENABLE_PROJECTS=llvm;clang;lld -DLLVM_INTEGRATED_CRT_ALLOC=../rpmalloc -DLLVM_USE_CRT_RELEASE=MT -DCMAKE_BUILD_TYPE=Release -DCMAKE_C_FLAGS_RELEASE="/O2 /Ob3 /DNDEBUG /Zi /Gy" -DCMAKE_CXX_FLAGS_RELEASE="/O2 /Ob3 /DNDEBUG /Zi /Gy" -DCMAKE_EXE_LINKER_FLAGS_RELEASE="/DEBUG /INCREMENTAL:NO /OPT:REF /OPT:ICF" -DCMAKE_INSTALL_PREFIX="..\install" -DLLVM_INCLUDE_BENCHMARKS=OFF -DLLVM_TARGETS_TO_BUILD="AArch64;X86" -DLLVM_ENABLE_ZLIB=FORCE_ON -DZLIB_LIBRARY="%ZLIB_PATH%/lib/Win64-llvm/Release/zlibstatic.lib" -DZLIB_INCLUDE_DIR="%ZLIB_PATH%/include/Win64/VS2015" -DCLANG_REPOSITORY_STRING="github.com/llvm/llvm-project" "..\source\llvm"
+%CMAKE_BINARY% -G %VS_VERSION% -DLLVM_ENABLE_PROJECTS=llvm;clang;lld -DLLVM_INTEGRATED_CRT_ALLOC=../rpmalloc -DCMAKE_MSVC_RUNTIME_LIBRARY=MultiThreaded -DCMAKE_BUILD_TYPE=Release -DCMAKE_C_FLAGS_RELEASE="/O2 /Ob3 /DNDEBUG /Zi /Gy" -DCMAKE_CXX_FLAGS_RELEASE="/O2 /Ob3 /DNDEBUG /Zi /Gy" -DCMAKE_EXE_LINKER_FLAGS_RELEASE="/DEBUG /INCREMENTAL:NO /OPT:REF /OPT:ICF" -DCMAKE_INSTALL_PREFIX="..\install" -DLLVM_INCLUDE_BENCHMARKS=OFF -DLLVM_TARGETS_TO_BUILD="AArch64;X86" -DLLVM_ENABLE_ZLIB=FORCE_ON -DZLIB_LIBRARY="%ZLIB_PATH%/lib/Win64/Release/zlibstatic.lib" -DZLIB_INCLUDE_DIR="%ZLIB_PATH%/include/" -DCLANG_REPOSITORY_STRING="github.com/llvm/llvm-project" "..\source\llvm"
+if ERRORLEVEL 1 goto endscript
 %CMAKE_BINARY% --build . --target install --config Release
+IF ERRORLEVEL 1 goto endscript
 popd
 
 for %%G in (aarch64-unknown-linux-gnueabi x86_64-unknown-linux-gnu) do (
@@ -76,18 +85,20 @@ for %%G in (aarch64-unknown-linux-gnueabi x86_64-unknown-linux-gnu) do (
     copy "install\bin\llvm-profdata.exe" OUTPUT\%%G\bin
     copy "install\bin\llvm-symbolizer.exe" OUTPUT\%%G\bin
     copy "install\bin\llvm-objcopy.exe" OUTPUT\%%G\bin
+    copy "install\bin\llvm-cov.exe" OUTPUT\%%G\bin
     copy "install\bin\LTO.dll" OUTPUT\%%G\bin
     xcopy "install\lib\clang" OUTPUT\%%G\lib\clang /s /e /y
 )
 
 @REM Create version file
-echo %TOOLCHAIN_VERSION%_clang-%LLVM_VERSION%-centos7> OUTPUT\ToolchainVersion.txt
+echo %TOOLCHAIN_VERSION%_clang-%LLVM_VERSION%-rockylinux8> OUTPUT\ToolchainVersion.txt
 
 echo Packing final toolchain...
 
 pushd OUTPUT
 del /S /Q %ROOT_DIR%\%FILENAME%.zip
-zip -r %ROOT_DIR%\%FILENAME%.zip *
+@REM the MS provided tar can be used to create and extract zip files on windows
+tar caf %ROOT_DIR%\%FILENAME%.zip *
 popd
 
 if exist "%NSIS_BINARY%" (
@@ -102,5 +113,6 @@ if exist "%NSIS_BINARY%" (
 
 popd
 
+:endscript
 echo.
 echo Done.
diff --git a/Engine/Build/BatchFiles/Linux/Toolchain/RunMe.sh b/Engine/Build/BatchFiles/Linux/Toolchain/RunMe.sh
--- a/Engine/Build/BatchFiles/Linux/Toolchain/RunMe.sh
+++ b/Engine/Build/BatchFiles/Linux/Toolchain/RunMe.sh
@@ -5,6 +5,10 @@ set -eu
 SCRIPT_DIR=$(cd "$(dirname "$BASH_SOURCE")" ; pwd)
 SCRIPT_NAME=$(basename "$BASH_SOURCE")
 
+Image=rockylinux:8
+
+DOCKER_BUILD_DIR=/src/build
+
 # https://stackoverflow.com/questions/23513045/how-to-check-if-a-process-is-running-inside-docker-container
 # https://unix.stackexchange.com/questions/607695/how-to-check-if-its-docker-or-host-machine-inside-bash-script
 # cgroups 2 busted the first way, attempting a new way here but may break again in the future
@@ -14,42 +18,70 @@ if ! [ -f "/.dockerenv" ]; then
   # host commands
   ##############################################################################
 
-  ZLIB_PATH=../../../../../Engine/Source/ThirdParty/zlib/v1.2.8/
+  ZLIB_PATH=../../../../../Engine/Source/ThirdParty/zlib/1.3/
 
   # Need to static link zlib for being able to compress debug files
   cp -rpvf $ZLIB_PATH ./
 
   ImageName=build_linux_toolchain
 
-  echo docker run -t --name ${ImageName} -v "${SCRIPT_DIR}:/src" centos:7 /src/${SCRIPT_NAME}
-  docker run -t --name ${ImageName} -v "${SCRIPT_DIR}:/src" centos:7 /src/${SCRIPT_NAME}
+  echo docker run -t --name ${ImageName} -v "${SCRIPT_DIR}:/src" ${Image} /src/${SCRIPT_NAME}
+  #docker run -t --name ${ImageName} -v "${SCRIPT_DIR}:/src" ${Image} /src/${SCRIPT_NAME}
 
   # Use if you want a shell when a command fails in the script
-  # docker run -it --name ${ImageName} -v "${SCRIPT_DIR}:/src" centos:7 bash -c "/src/${SCRIPT_NAME}; bash"
+  docker run -it --name ${ImageName} -v "${SCRIPT_DIR}:/src" ${Image} bash -c "/src/${SCRIPT_NAME}; bash"
 
-  echo Removing ${ImageName}...
+  #echo Removing ${ImageName}...
   docker rm ${ImageName}
 
 else
 
-  DOCKER_BUILD_DIR=/src/build
 
   if [ $UID -eq 0 ]; then
     ##############################################################################
     # docker root commands
     ##############################################################################
+
+
+
+  if [ "${Image}" == "centos:7" ]; then
+##### CentOs 7 stuff #####
+    echo "Running with Centos 7"
+    sleep 5
     yum install -y epel-release centos-release-scl dnf dnf-plugins-core
 
     # needed for mingw due to https://pagure.io/fesco/issue/2333
     dnf -y copr enable alonid/mingw-epel7
 
-    yum install -y ncurses-devel patch llvm-toolset-7 llvm-toolset-7-llvm-devel make cmake3 tree zip \
+    yum install -y ncurses-devel patch make tree zip \
         git wget which gcc-c++ gperf bison flex texinfo bzip2 help2man file unzip autoconf libtool \
         glibc-static libstdc++-devel libstdc++-static mingw64-gcc mingw64-gcc-c++ mingw64-winpthreads-static \
-        devtoolset-7-gcc libisl-devel
+        devtoolset-8-gcc.x86_64 devtoolset-8-gcc-c++.x86_64 libisl-devel openssl openssl-devel python3 zlib-static
+
+    # build a proper cmake
+    pushd ${DOCKER_BUILD_DIR}
+    if ! [ -f "cmake.done" ]; then
+        git clone [REDACTED_PROFILE_URL]/CMake.git
+
+        cd CMake
+        ./bootstrap --parallel=128 && make -j128 && make install && ln -s /usr/local/bin/cmake /usr/bin/cmake3
+        popd
+    fi
+
+##### End CentOs 7 stuff #####
+  else
+    echo "Running with Rockylinux 8"
+    sleep 5
+##### Rocky Linux 8 stuff #####
+	dnf install -y 'dnf-command(config-manager)'
+	dnf config-manager --enable powertools
+	yum install -y ncurses-devel patch make cmake3 zip git wget which gcc-c++ bzip2 file unzip autoconf libtool glibc-devel libstdc++-devel xz flex bison file unzip autoconf libtool diffutils openssl openssl-devel rsync perl-libintl perl-Text-Unidecode mingw32-gcc mingw32-gcc-c++ mingw64-gcc mingw64-gcc-c++ glibc-static mingw32-winpthreads-static.noarch mingw64-winpthreads-static.noarch libstdc++-static texinfo help2man perl-Unicode-Normalize gettext python3
+
+##### End Rocky Linux 8 stuff #####
+  fi
 
     # Create non-privileged user and workspace
-    adduser buildmaster
+    adduser buildmaster || echo "User exists"
     mkdir -p ${DOCKER_BUILD_DIR}
     chown buildmaster:nobody -R ${DOCKER_BUILD_DIR}
 
diff --git a/Engine/Build/BatchFiles/Linux/Toolchain/aarch64-unknown-linux-gnueabi.linux.config b/Engine/Build/BatchFiles/Linux/Toolchain/aarch64-unknown-linux-gnueabi.linux.config
--- a/Engine/Build/BatchFiles/Linux/Toolchain/aarch64-unknown-linux-gnueabi.linux.config
+++ b/Engine/Build/BatchFiles/Linux/Toolchain/aarch64-unknown-linux-gnueabi.linux.config
@@ -1,10 +1,34 @@
 #
 # Automatically generated file; DO NOT EDIT.
-# Crosstool-NG Configuration
-#
-CT_CONFIGURE_has_make381=y
-CT_CONFIGURE_has_xz=y
-CT_CONFIGURE_has_cvs=y
+# crosstool-NG 1.26.0 Configuration
+#
+CT_CONFIGURE_has_static_link=y
+CT_CONFIGURE_has_cxx11=y
+CT_CONFIGURE_has_wget=y
+CT_CONFIGURE_has_curl=y
+CT_CONFIGURE_has_rsync=y
+CT_CONFIGURE_has_make_3_81_or_newer=y
+CT_CONFIGURE_has_make_4_0_or_newer=y
+CT_CONFIGURE_has_libtool_2_4_or_newer=y
+CT_CONFIGURE_has_libtoolize_2_4_or_newer=y
+CT_CONFIGURE_has_autoconf_2_65_or_newer=y
+CT_CONFIGURE_has_autoreconf_2_65_or_newer=y
+CT_CONFIGURE_has_automake_1_15_or_newer=y
+CT_CONFIGURE_has_gnu_m4_1_4_12_or_newer=y
+CT_CONFIGURE_has_bison_2_7_or_newer=y
+CT_CONFIGURE_has_bison_3_0_4_or_newer=y
+CT_CONFIGURE_has_python=y
+CT_CONFIGURE_has_git=y
+CT_CONFIGURE_has_md5sum=y
+CT_CONFIGURE_has_sha1sum=y
+CT_CONFIGURE_has_sha256sum=y
+CT_CONFIGURE_has_sha512sum=y
+CT_CONFIGURE_has_install_with_strip_program=y
+CT_VERSION="1.26.0"
+CT_VCHECK=""
+CT_CONFIG_VERSION_ENV="4"
+CT_CONFIG_VERSION_CURRENT="4"
+CT_CONFIG_VERSION="4"
 CT_MODULES=y
 
 #
@@ -21,41 +45,49 @@ CT_MODULES=y
 #
 # Paths
 #
-CT_LOCAL_TARBALLS_DIR=""
+CT_LOCAL_TARBALLS_DIR="${HOME}/src"
+CT_SAVE_TARBALLS=y
+# CT_TARBALLS_BUILDROOT_LAYOUT is not set
 CT_WORK_DIR="${CT_TOP_DIR}/.build"
+CT_BUILD_TOP_DIR="${CT_WORK_DIR:-${CT_TOP_DIR}/.build}/${CT_HOST:+HOST-${CT_HOST}/}${CT_TARGET}"
+CT_BUILD_DIR="${CT_BUILD_TOP_DIR}/build"
 CT_PREFIX_DIR="${CT_TOP_DIR}/../OUTPUT-linux/${CT_TARGET}"
-CT_INSTALL_DIR="${CT_PREFIX_DIR}"
 CT_RM_RF_PREFIX_DIR=y
 CT_REMOVE_DOCS=y
-CT_INSTALL_DIR_RO=y
+CT_INSTALL_LICENSES=y
+CT_PREFIX_DIR_RO=y
 CT_STRIP_HOST_TOOLCHAIN_EXECUTABLES=y
 # CT_STRIP_TARGET_TOOLCHAIN_EXECUTABLES is not set
 
 #
 # Downloading
 #
+CT_DOWNLOAD_AGENT_WGET=y
+# CT_DOWNLOAD_AGENT_CURL is not set
+# CT_DOWNLOAD_AGENT_NONE is not set
 # CT_FORBID_DOWNLOAD is not set
 # CT_FORCE_DOWNLOAD is not set
 CT_CONNECT_TIMEOUT=10
+CT_DOWNLOAD_WGET_OPTIONS="--passive-ftp --tries=3 -nc --progress=dot:binary"
 # CT_ONLY_DOWNLOAD is not set
 # CT_USE_MIRROR is not set
+CT_VERIFY_DOWNLOAD_DIGEST=y
+CT_VERIFY_DOWNLOAD_DIGEST_SHA512=y
+# CT_VERIFY_DOWNLOAD_DIGEST_SHA256 is not set
+# CT_VERIFY_DOWNLOAD_DIGEST_SHA1 is not set
+# CT_VERIFY_DOWNLOAD_DIGEST_MD5 is not set
+CT_VERIFY_DOWNLOAD_DIGEST_ALG="sha512"
+# CT_VERIFY_DOWNLOAD_SIGNATURE is not set
 
 #
 # Extracting
 #
 # CT_FORCE_EXTRACT is not set
-CT_OVERIDE_CONFIG_GUESS_SUB=y
+CT_OVERRIDE_CONFIG_GUESS_SUB=y
 # CT_ONLY_EXTRACT is not set
-# CT_PATCH_BUNDLED is not set
-# CT_PATCH_LOCAL is not set
-CT_PATCH_BUNDLED_LOCAL=y
-# CT_PATCH_LOCAL_BUNDLED is not set
-# CT_PATCH_BUNDLED_FALLBACK_LOCAL is not set
-# CT_PATCH_LOCAL_FALLBACK_BUNDLED is not set
-# CT_PATCH_NONE is not set
-CT_PATCH_ORDER="bundled,local"
-CT_PATCH_USE_LOCAL=y
-CT_LOCAL_PATCH_DIR="${CT_TOP_DIR}/../../patches"
+CT_PATCH_BUNDLED=y
+# CT_PATCH_BUNDLED_LOCAL is not set
+CT_PATCH_ORDER="bundled"
 
 #
 # Build behavior
@@ -64,6 +96,7 @@ CT_PARALLEL_JOBS=0
 CT_LOAD=""
 CT_USE_PIPES=y
 CT_EXTRA_CFLAGS_FOR_BUILD=""
+CT_EXTRA_CXXFLAGS_FOR_BUILD=""
 CT_EXTRA_LDFLAGS_FOR_BUILD=""
 CT_EXTRA_CFLAGS_FOR_HOST=""
 CT_EXTRA_LDFLAGS_FOR_HOST=""
@@ -78,93 +111,97 @@ CT_CONFIG_SHELL="${bash}"
 #
 # CT_LOG_ERROR is not set
 # CT_LOG_WARN is not set
-CT_LOG_INFO=y
-# CT_LOG_EXTRA is not set
+# CT_LOG_INFO is not set
+CT_LOG_EXTRA=y
 # CT_LOG_ALL is not set
 # CT_LOG_DEBUG is not set
-CT_LOG_LEVEL_MAX="INFO"
+CT_LOG_LEVEL_MAX="EXTRA"
 # CT_LOG_SEE_TOOLS_WARN is not set
 CT_LOG_PROGRESS_BAR=y
 CT_LOG_TO_FILE=y
 CT_LOG_FILE_COMPRESS=y
+# end of Paths and misc options
 
 #
 # Target options
 #
+# CT_ARCH_ALPHA is not set
+# CT_ARCH_ARC is not set
+CT_ARCH_ARM=y
+# CT_ARCH_AVR is not set
+# CT_ARCH_BPF is not set
+# CT_ARCH_M68K is not set
+# CT_ARCH_MIPS is not set
+# CT_ARCH_NIOS2 is not set
+# CT_ARCH_POWERPC is not set
+# CT_ARCH_PRU is not set
+# CT_ARCH_S390 is not set
+# CT_ARCH_SH is not set
+# CT_ARCH_SPARC is not set
+# CT_ARCH_X86 is not set
+# CT_ARCH_XTENSA is not set
 CT_ARCH="arm"
-CT_ARCH_SUPPORTS_BOTH_MMU=y
-CT_ARCH_SUPPORTS_BOTH_ENDIAN=y
-CT_ARCH_SUPPORTS_32=y
-CT_ARCH_SUPPORTS_64=y
-CT_ARCH_SUPPORTS_WITH_ARCH=y
-CT_ARCH_SUPPORTS_WITH_CPU=y
-CT_ARCH_SUPPORTS_WITH_TUNE=y
-CT_ARCH_DEFAULT_HAS_MMU=y
-CT_ARCH_DEFAULT_LE=y
-CT_ARCH_DEFAULT_32=y
-CT_ARCH_ARCH="armv8-a"
+CT_ARCH_CHOICE_KSYM="ARM"
 CT_ARCH_CPU=""
 CT_ARCH_TUNE=""
-# CT_ARCH_BE is not set
-CT_ARCH_LE=y
-# CT_ARCH_32 is not set
-CT_ARCH_64=y
-CT_ARCH_BITNESS=64
-# CT_ARCH_FLOAT_HW is not set
-# CT_ARCH_FLOAT_SW is not set
-CT_TARGET_CFLAGS=""
-CT_TARGET_LDFLAGS=""
-# CT_ARCH_alpha is not set
-CT_ARCH_arm=y
-# CT_ARCH_avr is not set
-# CT_ARCH_m68k is not set
-# CT_ARCH_mips is not set
-# CT_ARCH_nios2 is not set
-# CT_ARCH_powerpc is not set
-# CT_ARCH_s390 is not set
-# CT_ARCH_sh is not set
-# CT_ARCH_sparc is not set
-# CT_ARCH_x86 is not set
-# CT_ARCH_xtensa is not set
-CT_ARCH_alpha_AVAILABLE=y
-CT_ARCH_arm_AVAILABLE=y
-CT_ARCH_avr_AVAILABLE=y
-CT_ARCH_m68k_AVAILABLE=y
-CT_ARCH_microblaze_AVAILABLE=y
-CT_ARCH_mips_AVAILABLE=y
-CT_ARCH_nios2_AVAILABLE=y
-CT_ARCH_powerpc_AVAILABLE=y
-CT_ARCH_s390_AVAILABLE=y
-CT_ARCH_sh_AVAILABLE=y
-CT_ARCH_sparc_AVAILABLE=y
-CT_ARCH_x86_AVAILABLE=y
-CT_ARCH_xtensa_AVAILABLE=y
+CT_ARCH_ARM_SHOW=y
+
+#
+# Options for arm
+#
+CT_ARCH_ARM_PKG_KSYM=""
+CT_ARCH_ARM_MODE="arm"
+CT_ARCH_ARM_MODE_ARM=y
+# CT_ARCH_ARM_MODE_THUMB is not set
+# CT_ARCH_ARM_INTERWORKING is not set
+CT_ARCH_ARM_EABI_FORCE=y
+CT_ARCH_ARM_EABI=y
+CT_ALL_ARCH_CHOICES="ALPHA ARC ARM AVR BPF C6X LOONGARCH M68K MICROBLAZE MIPS MOXIE MSP430 NIOS2 POWERPC PRU RISCV S390 SH SPARC X86 XTENSA"
 CT_ARCH_SUFFIX=""
+# CT_OMIT_TARGET_VENDOR is not set
 
 #
 # Generic target options
 #
 # CT_MULTILIB is not set
+CT_DEMULTILIB=y
+CT_ARCH_SUPPORTS_BOTH_MMU=y
+CT_ARCH_DEFAULT_HAS_MMU=y
 CT_ARCH_USE_MMU=y
+CT_ARCH_SUPPORTS_FLAT_FORMAT=y
+CT_ARCH_SUPPORTS_LIBSANITIZER=y
+CT_ARCH_SUPPORTS_EITHER_ENDIAN=y
+CT_ARCH_DEFAULT_LE=y
+# CT_ARCH_BE is not set
+CT_ARCH_LE=y
 CT_ARCH_ENDIAN="little"
+CT_ARCH_SUPPORTS_32=y
+CT_ARCH_SUPPORTS_64=y
+CT_ARCH_DEFAULT_32=y
+CT_ARCH_BITNESS=64
+CT_ARCH_32=y
+CT_ARCH_64=y
 
 #
 # Target optimisations
 #
+CT_ARCH_SUPPORTS_WITH_ARCH=y
+CT_ARCH_SUPPORTS_WITH_CPU=y
+CT_ARCH_SUPPORTS_WITH_TUNE=y
+CT_ARCH_SUPPORTS_WITH_FLOAT=y
+CT_ARCH_SUPPORTS_WITH_FPU=y
+CT_ARCH_SUPPORTS_SOFTFP=y
 CT_ARCH_EXCLUSIVE_WITH_CPU=y
-# CT_ARCH_FLOAT_AUTO is not set
+CT_ARCH_ARCH="armv8-a"
+CT_ARCH_FPU=""
+CT_ARCH_FLOAT_AUTO=y
+# CT_ARCH_FLOAT_HW is not set
 # CT_ARCH_FLOAT_SOFTFP is not set
-CT_ARCH_FLOAT=""
-
-#
-# arm other options
-#
-CT_ARCH_ARM_MODE="arm"
-CT_ARCH_ARM_MODE_ARM=y
-# CT_ARCH_ARM_MODE_THUMB is not set
-# CT_ARCH_ARM_INTERWORKING is not set
-CT_ARCH_ARM_EABI_FORCE=y
-CT_ARCH_ARM_EABI=y
+# CT_ARCH_FLOAT_SW is not set
+CT_TARGET_CFLAGS=""
+CT_TARGET_LDFLAGS=""
+CT_ARCH_FLOAT="auto"
+# end of Target options
 
 #
 # Toolchain options
@@ -178,7 +215,9 @@ CT_USE_SYSROOT=y
 CT_SYSROOT_NAME="sysroot"
 CT_SYSROOT_DIR_PREFIX=""
 CT_WANTS_STATIC_LINK=y
+CT_WANTS_STATIC_LINK_CXX=y
 CT_STATIC_TOOLCHAIN=y
+CT_SHOW_CT_VERSION=y
 CT_TOOLCHAIN_PKGVERSION=""
 CT_TOOLCHAIN_BUGURL=""
 
@@ -207,132 +246,279 @@ CT_BUILD_SUFFIX=""
 # Misc options
 #
 # CT_TOOLCHAIN_ENABLE_NLS is not set
+# end of Toolchain options
 
 #
 # Operating System
 #
 CT_KERNEL_SUPPORTS_SHARED_LIBS=y
+# CT_KERNEL_BARE_METAL is not set
+CT_KERNEL_LINUX=y
 CT_KERNEL="linux"
-CT_KERNEL_VERSION="3.10.93"
-# CT_KERNEL_bare_metal is not set
-CT_KERNEL_linux=y
-CT_KERNEL_bare_metal_AVAILABLE=y
-CT_KERNEL_linux_AVAILABLE=y
-# CT_KERNEL_V_4_3 is not set
-# CT_KERNEL_V_4_2 is not set
-# CT_KERNEL_V_4_1 is not set
-# CT_KERNEL_V_3_18 is not set
-# CT_KERNEL_V_3_14 is not set
-# CT_KERNEL_V_3_12 is not set
-CT_KERNEL_V_3_10=y
-# CT_KERNEL_V_3_4 is not set
-# CT_KERNEL_V_3_2 is not set
-# CT_KERNEL_V_2_6_32 is not set
-# CT_KERNEL_LINUX_CUSTOM is not set
-CT_KERNEL_windows_AVAILABLE=y
-
-#
-# Common kernel options
-#
-CT_SHARED_LIBS=y
-
-#
-# linux other options
-#
+CT_KERNEL_CHOICE_KSYM="LINUX"
+CT_KERNEL_LINUX_SHOW=y
+
+#
+# Options for linux
+#
+CT_KERNEL_LINUX_PKG_KSYM="LINUX"
+CT_LINUX_DIR_NAME="linux"
+CT_LINUX_USE_WWW_KERNEL_ORG=y
+# CT_LINUX_USE_ORACLE is not set
+CT_LINUX_USE="LINUX"
+CT_LINUX_PKG_NAME="linux"
+CT_LINUX_SRC_RELEASE=y
+# CT_LINUX_SRC_DEVEL is not set
+CT_LINUX_PATCH_ORDER="global"
+# CT_LINUX_V_6_4 is not set
+# CT_LINUX_V_6_3 is not set
+# CT_LINUX_V_6_2 is not set
+# CT_LINUX_V_6_1 is not set
+# CT_LINUX_V_6_0 is not set
+# CT_LINUX_V_5_19 is not set
+# CT_LINUX_V_5_18 is not set
+# CT_LINUX_V_5_17 is not set
+# CT_LINUX_V_5_16 is not set
+# CT_LINUX_V_5_15 is not set
+# CT_LINUX_V_5_14 is not set
+# CT_LINUX_V_5_13 is not set
+# CT_LINUX_V_5_12 is not set
+# CT_LINUX_V_5_11 is not set
+# CT_LINUX_V_5_10 is not set
+# CT_LINUX_V_5_9 is not set
+# CT_LINUX_V_5_8 is not set
+# CT_LINUX_V_5_7 is not set
+# CT_LINUX_V_5_4 is not set
+# CT_LINUX_V_5_3 is not set
+# CT_LINUX_V_5_2 is not set
+# CT_LINUX_V_5_1 is not set
+# CT_LINUX_V_5_0 is not set
+# CT_LINUX_V_4_20 is not set
+# CT_LINUX_V_4_19 is not set
+CT_LINUX_V_4_18=y
+# CT_LINUX_V_4_17 is not set
+# CT_LINUX_V_4_16 is not set
+# CT_LINUX_V_4_15 is not set
+# CT_LINUX_V_4_14 is not set
+# CT_LINUX_V_4_13 is not set
+# CT_LINUX_V_4_12 is not set
+# CT_LINUX_V_4_11 is not set
+# CT_LINUX_V_4_10 is not set
+# CT_LINUX_V_4_9 is not set
+# CT_LINUX_V_4_4 is not set
+# CT_LINUX_V_4_1 is not set
+# CT_LINUX_V_3_16 is not set
+# CT_LINUX_V_3_13 is not set
+# CT_LINUX_V_3_12 is not set
+# CT_LINUX_V_3_10 is not set
+# CT_LINUX_V_3_4 is not set
+# CT_LINUX_V_3_2 is not set
+CT_LINUX_VERSION="4.18.20"
+CT_LINUX_MIRRORS="$(CT_Mirrors kernel.org linux ${CT_LINUX_VERSION})"
+CT_LINUX_ARCHIVE_FILENAME="@{pkg_name}-@{version}"
+CT_LINUX_ARCHIVE_DIRNAME="@{pkg_name}-@{version}"
+CT_LINUX_ARCHIVE_FORMATS=".tar.xz .tar.gz"
+CT_LINUX_SIGNATURE_FORMAT="unpacked/.sign"
+CT_LINUX_5_19_or_older=y
+CT_LINUX_older_than_5_19=y
+CT_LINUX_5_12_or_older=y
+CT_LINUX_older_than_5_12=y
+CT_LINUX_5_5_or_older=y
+CT_LINUX_older_than_5_5=y
+CT_LINUX_5_3_or_older=y
+CT_LINUX_older_than_5_3=y
+CT_LINUX_later_than_4_8=y
+CT_LINUX_4_8_or_later=y
+CT_LINUX_later_than_3_7=y
+CT_LINUX_3_7_or_later=y
+CT_LINUX_later_than_3_2=y
+CT_LINUX_3_2_or_later=y
+CT_LINUX_REQUIRE_3_2_or_later=y
+CT_KERNEL_DEP_RSYNC=y
 CT_KERNEL_LINUX_VERBOSITY_0=y
 # CT_KERNEL_LINUX_VERBOSITY_1 is not set
 # CT_KERNEL_LINUX_VERBOSITY_2 is not set
 CT_KERNEL_LINUX_VERBOSE_LEVEL=0
 CT_KERNEL_LINUX_INSTALL_CHECK=y
+CT_ALL_KERNEL_CHOICES="BARE_METAL LINUX WINDOWS"
+
+#
+# Common kernel options
+#
+CT_SHARED_LIBS=y
+# end of Operating System
 
 #
 # Binary utilities
 #
 CT_ARCH_BINFMT_ELF=y
+CT_BINUTILS_BINUTILS=y
 CT_BINUTILS="binutils"
-CT_BINUTILS_binutils=y
+CT_BINUTILS_CHOICE_KSYM="BINUTILS"
+CT_BINUTILS_BINUTILS_SHOW=y
+
+#
+# Options for binutils
+#
+CT_BINUTILS_BINUTILS_PKG_KSYM="BINUTILS"
+CT_BINUTILS_DIR_NAME="binutils"
+CT_BINUTILS_USE_GNU=y
+# CT_BINUTILS_USE_ORACLE is not set
+CT_BINUTILS_USE="BINUTILS"
+CT_BINUTILS_PKG_NAME="binutils"
+CT_BINUTILS_SRC_RELEASE=y
+# CT_BINUTILS_SRC_DEVEL is not set
+CT_BINUTILS_PATCH_ORDER="global"
+CT_BINUTILS_V_2_40=y
+# CT_BINUTILS_V_2_39 is not set
+# CT_BINUTILS_V_2_38 is not set
+# CT_BINUTILS_V_2_37 is not set
+# CT_BINUTILS_V_2_36 is not set
+# CT_BINUTILS_V_2_35 is not set
+# CT_BINUTILS_V_2_34 is not set
+# CT_BINUTILS_V_2_33 is not set
+# CT_BINUTILS_V_2_32 is not set
+# CT_BINUTILS_V_2_31 is not set
+# CT_BINUTILS_V_2_30 is not set
+# CT_BINUTILS_V_2_29 is not set
+# CT_BINUTILS_V_2_28 is not set
+# CT_BINUTILS_V_2_27 is not set
+# CT_BINUTILS_V_2_26 is not set
+CT_BINUTILS_VERSION="2.40"
+CT_BINUTILS_MIRRORS="$(CT_Mirrors GNU binutils) $(CT_Mirrors sourceware binutils/releases)"
+CT_BINUTILS_ARCHIVE_FILENAME="@{pkg_name}-@{version}"
+CT_BINUTILS_ARCHIVE_DIRNAME="@{pkg_name}-@{version}"
+CT_BINUTILS_ARCHIVE_FORMATS=".tar.xz .tar.bz2 .tar.gz"
+CT_BINUTILS_SIGNATURE_FORMAT="packed/.sig"
+CT_BINUTILS_later_than_2_39=y
+CT_BINUTILS_2_39_or_later=y
+CT_BINUTILS_later_than_2_30=y
+CT_BINUTILS_2_30_or_later=y
+CT_BINUTILS_later_than_2_27=y
+CT_BINUTILS_2_27_or_later=y
+CT_BINUTILS_later_than_2_26=y
+CT_BINUTILS_2_26_or_later=y
 
 #
 # GNU binutils
 #
-# CT_CC_BINUTILS_SHOW_LINARO is not set
-CT_BINUTILS_V_2_25_1=y
-# CT_BINUTILS_V_2_25 is not set
-# CT_BINUTILS_V_2_24 is not set
-# CT_BINUTILS_V_2_23_2 is not set
-# CT_BINUTILS_V_2_23_1 is not set
-# CT_BINUTILS_V_2_22 is not set
-# CT_BINUTILS_V_2_21_53 is not set
-# CT_BINUTILS_V_2_21_1a is not set
-# CT_BINUTILS_V_2_20_1a is not set
-# CT_BINUTILS_V_2_19_1a is not set
-# CT_BINUTILS_V_2_18a is not set
-CT_BINUTILS_VERSION="2.25.1"
-CT_BINUTILS_2_25_1_or_later=y
-CT_BINUTILS_2_25_or_later=y
-CT_BINUTILS_2_24_or_later=y
-CT_BINUTILS_2_23_or_later=y
-CT_BINUTILS_2_22_or_later=y
-CT_BINUTILS_2_21_or_later=y
-CT_BINUTILS_2_20_or_later=y
-CT_BINUTILS_2_19_or_later=y
-CT_BINUTILS_2_18_or_later=y
-CT_BINUTILS_HAS_HASH_STYLE=y
-CT_BINUTILS_HAS_GOLD=y
 CT_BINUTILS_GOLD_SUPPORTS_ARCH=y
-CT_BINUTILS_GOLD_SUPPORT=y
-CT_BINUTILS_HAS_PLUGINS=y
-CT_BINUTILS_HAS_PKGVERSION_BUGURL=y
-CT_BINUTILS_FORCE_LD_BFD=y
+CT_BINUTILS_FORCE_LD_BFD_DEFAULT=y
 CT_BINUTILS_LINKER_LD=y
-# CT_BINUTILS_LINKER_LD_GOLD is not set
-# CT_BINUTILS_LINKER_GOLD_LD is not set
 CT_BINUTILS_LINKERS_LIST="ld"
 CT_BINUTILS_LINKER_DEFAULT="bfd"
-# CT_BINUTILS_PLUGINS is not set
+CT_BINUTILS_RELRO=m
+CT_BINUTILS_DETERMINISTIC_ARCHIVES=y
 CT_BINUTILS_EXTRA_CONFIG_ARRAY=""
 # CT_BINUTILS_FOR_TARGET is not set
-
-#
-# binutils other options
-#
+CT_ALL_BINUTILS_CHOICES="BINUTILS"
+# end of Binary utilities
 
 #
 # C-library
 #
+CT_LIBC_GLIBC=y
+# CT_LIBC_UCLIBC_NG is not set
 CT_LIBC="glibc"
-CT_LIBC_VERSION="2.17"
-CT_LIBC_glibc=y
-# CT_LIBC_musl is not set
-# CT_LIBC_uClibc is not set
-CT_LIBC_avr_libc_AVAILABLE=y
-CT_LIBC_glibc_AVAILABLE=y
+CT_LIBC_CHOICE_KSYM="GLIBC"
+CT_LIBC_GLIBC_SHOW=y
+
+#
+# Options for glibc
+#
+CT_LIBC_GLIBC_PKG_KSYM="GLIBC"
+CT_GLIBC_DIR_NAME="glibc"
+CT_GLIBC_USE_GNU=y
+# CT_GLIBC_USE_ORACLE is not set
+CT_GLIBC_USE="GLIBC"
+CT_GLIBC_PKG_NAME="glibc"
+CT_GLIBC_SRC_RELEASE=y
+# CT_GLIBC_SRC_DEVEL is not set
+CT_GLIBC_PATCH_ORDER="global"
+CT_GLIBC_V_2_28=y
+# CT_GLIBC_V_2_27 is not set
+# CT_GLIBC_V_2_26 is not set
+# CT_GLIBC_V_2_25 is not set
+# CT_GLIBC_V_2_24 is not set
+# CT_GLIBC_V_2_23 is not set
+# CT_GLIBC_V_2_19 is not set
+# CT_GLIBC_V_2_17 is not set
+CT_GLIBC_VERSION="2.28"
+CT_GLIBC_MIRRORS="$(CT_Mirrors GNU glibc)"
+CT_GLIBC_ARCHIVE_FILENAME="@{pkg_name}-@{version}"
+CT_GLIBC_ARCHIVE_DIRNAME="@{pkg_name}-@{version}"
+CT_GLIBC_ARCHIVE_FORMATS=".tar.xz .tar.bz2 .tar.gz"
+CT_GLIBC_SIGNATURE_FORMAT="packed/.sig"
+CT_GLIBC_2_38_or_older=y
+CT_GLIBC_older_than_2_38=y
+CT_GLIBC_2_37_or_older=y
+CT_GLIBC_older_than_2_37=y
+CT_GLIBC_2_36_or_older=y
+CT_GLIBC_older_than_2_36=y
+CT_GLIBC_2_34_or_older=y
+CT_GLIBC_older_than_2_34=y
+CT_GLIBC_2_32_or_older=y
+CT_GLIBC_older_than_2_32=y
+CT_GLIBC_2_31_or_older=y
+CT_GLIBC_older_than_2_31=y
+CT_GLIBC_2_30_or_older=y
+CT_GLIBC_older_than_2_30=y
+CT_GLIBC_2_29_or_older=y
+CT_GLIBC_older_than_2_29=y
+CT_GLIBC_REQUIRE_older_than_2_29=y
+CT_GLIBC_2_28_or_later=y
+CT_GLIBC_2_28_or_older=y
+CT_GLIBC_later_than_2_27=y
+CT_GLIBC_2_27_or_later=y
+CT_GLIBC_later_than_2_26=y
+CT_GLIBC_2_26_or_later=y
+CT_GLIBC_later_than_2_25=y
+CT_GLIBC_2_25_or_later=y
+CT_GLIBC_later_than_2_24=y
+CT_GLIBC_2_24_or_later=y
+CT_GLIBC_later_than_2_23=y
+CT_GLIBC_2_23_or_later=y
+CT_GLIBC_later_than_2_20=y
+CT_GLIBC_2_20_or_later=y
+CT_GLIBC_later_than_2_17=y
+CT_GLIBC_2_17_or_later=y
+CT_GLIBC_later_than_2_14=y
+CT_GLIBC_2_14_or_later=y
+CT_GLIBC_DEP_KERNEL_HEADERS_VERSION=y
+CT_GLIBC_DEP_BINUTILS=y
+CT_GLIBC_DEP_GCC=y
+CT_GLIBC_DEP_PYTHON=y
+
+#
+# glibc >=2.29 requires python >3.4
+#
+CT_GLIBC_SPARC_ALLOW_V7=y
 CT_THREADS="nptl"
-# CT_CC_GLIBC_SHOW_LINARO is not set
-# CT_LIBC_GLIBC_V_2_22 is not set
-# CT_LIBC_GLIBC_V_2_21 is not set
-# CT_LIBC_GLIBC_V_2_20 is not set
-# CT_LIBC_GLIBC_V_2_19 is not set
-# CT_LIBC_GLIBC_V_2_18 is not set
-CT_LIBC_GLIBC_V_2_17=y
-# CT_LIBC_GLIBC_V_2_16_0 is not set
-# CT_LIBC_GLIBC_V_2_15 is not set
-# CT_LIBC_GLIBC_V_2_14_1 is not set
-# CT_LIBC_GLIBC_V_2_14 is not set
-# CT_LIBC_GLIBC_V_2_13 is not set
-# CT_LIBC_GLIBC_V_2_12_2 is not set
-# CT_LIBC_GLIBC_V_2_12_1 is not set
-# CT_LIBC_GLIBC_V_2_11_1 is not set
-# CT_LIBC_GLIBC_V_2_11 is not set
-# CT_LIBC_GLIBC_V_2_10_1 is not set
-# CT_LIBC_GLIBC_V_2_9 is not set
-# CT_LIBC_GLIBC_V_2_8 is not set
-CT_LIBC_GLIBC_2_17_or_later=y
-CT_LIBC_mingw_AVAILABLE=y
-CT_LIBC_musl_AVAILABLE=y
-CT_LIBC_newlib_AVAILABLE=y
-CT_LIBC_none_AVAILABLE=y
-CT_LIBC_uClibc_AVAILABLE=y
+CT_GLIBC_BUILD_SSP=y
+CT_GLIBC_HAS_LIBIDN_ADDON=y
+# CT_GLIBC_USE_LIBIDN_ADDON is not set
+CT_GLIBC_NO_SPARC_V8=y
+CT_GLIBC_HAS_OBSOLETE_RPC=y
+CT_GLIBC_EXTRA_CONFIG_ARRAY=""
+CT_GLIBC_CONFIGPARMS=""
+CT_GLIBC_ENABLE_DEBUG=y
+CT_GLIBC_EXTRA_CFLAGS="-Wno-missing-attributes -Wno-array-bounds -Wno-array-parameter -Wno-stringop-overflow -Wno-maybe-uninitialized"
+CT_GLIBC_ENABLE_OBSOLETE_RPC=y
+# CT_GLIBC_DISABLE_VERSIONING is not set
+CT_GLIBC_OLDEST_ABI=""
+CT_GLIBC_FORCE_UNWIND=y
+# CT_GLIBC_LOCALES is not set
+# CT_GLIBC_KERNEL_VERSION_NONE is not set
+CT_GLIBC_KERNEL_VERSION_AS_HEADERS=y
+# CT_GLIBC_KERNEL_VERSION_CHOSEN is not set
+CT_GLIBC_MIN_KERNEL="4.18.20"
+CT_GLIBC_SSP_DEFAULT=y
+# CT_GLIBC_SSP_NO is not set
+# CT_GLIBC_SSP_YES is not set
+# CT_GLIBC_SSP_ALL is not set
+# CT_GLIBC_SSP_STRONG is not set
+CT_GLIBC_ENABLE_COMMON_FLAG=y
+CT_ALL_LIBC_CHOICES="AVR_LIBC GLIBC MINGW_W64 MOXIEBOX MUSL NEWLIB NONE PICOLIBC UCLIBC_NG"
 CT_LIBC_SUPPORT_THREADS_ANY=y
 CT_LIBC_SUPPORT_THREADS_NATIVE=y
 
@@ -340,122 +526,107 @@ CT_LIBC_SUPPORT_THREADS_NATIVE=y
 # Common C library options
 #
 CT_THREADS_NATIVE=y
+# CT_CREATE_LDSO_CONF is not set
 CT_LIBC_XLDD=y
-
-#
-# glibc other options
-#
-# CT_LIBC_GLIBC_PORTS_EXTERNAL is not set
-CT_LIBC_GLIBC_MAY_FORCE_PORTS=y
-CT_LIBC_glibc_familly=y
-CT_LIBC_GLIBC_EXTRA_CONFIG_ARRAY=""
-CT_LIBC_GLIBC_CONFIGPARMS=""
-CT_LIBC_GLIBC_EXTRA_CFLAGS=""
-CT_LIBC_EXTRA_CC_ARGS=""
-# CT_LIBC_DISABLE_VERSIONING is not set
-CT_LIBC_OLDEST_ABI=""
-CT_LIBC_GLIBC_FORCE_UNWIND=y
-CT_LIBC_GLIBC_USE_PORTS=y
-CT_LIBC_ADDONS_LIST=""
-
-#
-# WARNING !!!
-#
-
-#
-#   For glibc >= 2.8, it can happen that the tarballs
-#
-
-#
-#   for the addons are not available for download.
-#
-
-#
-#   If that happens, bad luck... Try a previous version
-#
-
-#
-#   or try again later... :-(
-#
-# CT_LIBC_LOCALES is not set
-# CT_LIBC_GLIBC_KERNEL_VERSION_NONE is not set
-CT_LIBC_GLIBC_KERNEL_VERSION_AS_HEADERS=y
-# CT_LIBC_GLIBC_KERNEL_VERSION_CHOSEN is not set
-CT_LIBC_GLIBC_MIN_KERNEL="3.10.93"
+# end of C-library
 
 #
 # C compiler
 #
+CT_CC_CORE_NEEDED=y
+CT_CC_SUPPORT_CXX=y
+CT_CC_SUPPORT_FORTRAN=y
+CT_CC_SUPPORT_ADA=y
+CT_CC_SUPPORT_D=y
+CT_CC_SUPPORT_JIT=y
+CT_CC_SUPPORT_OBJC=y
+CT_CC_SUPPORT_OBJCXX=y
+CT_CC_SUPPORT_GOLANG=y
+CT_CC_GCC=y
 CT_CC="gcc"
-CT_CC_CORE_PASSES_NEEDED=y
-CT_CC_CORE_PASS_1_NEEDED=y
-CT_CC_CORE_PASS_2_NEEDED=y
-CT_CC_gcc=y
-CT_CC_GCC_SHOW_LINARO=y
-# CT_CC_GCC_V_5_2_0 is not set
-# CT_CC_GCC_V_linaro_4_9=y
-# CT_CC_GCC_V_4_9_3 is not set
-CT_CC_GCC_V_linaro_4_8=y
-# CT_CC_GCC_V_4_8_5 is not set
-# CT_CC_GCC_V_linaro_4_7 is not set
-# CT_CC_GCC_V_4_7_4 is not set
-# CT_CC_GCC_V_linaro_4_6 is not set
-# CT_CC_GCC_V_4_6_4 is not set
-# CT_CC_GCC_V_linaro_4_5 is not set
-# CT_CC_GCC_V_4_5_4 is not set
-# CT_CC_GCC_V_linaro_4_4 is not set
-# CT_CC_GCC_V_4_4_7 is not set
-# CT_CC_GCC_V_4_3_6 is not set
-# CT_CC_GCC_V_4_2_4 is not set
-# CT_CC_GCC_CUSTOM is not set
-CT_CC_GCC_4_2_or_later=y
-CT_CC_GCC_4_3_or_later=y
-CT_CC_GCC_4_4_or_later=y
-CT_CC_GCC_4_5_or_later=y
-CT_CC_GCC_4_6_or_later=y
-CT_CC_GCC_4_7_or_later=y
-CT_CC_GCC_4_8_or_later=y
-CT_CC_GCC_4_8=y
-CT_CC_GCC_4_9_or_later=y
-CT_CC_GCC_HAS_GRAPHITE=y
-# CT_CC_GCC_USE_GRAPHITE is not set
-CT_CC_GCC_HAS_LTO=y
-# CT_CC_GCC_USE_LTO is not set
-CT_CC_GCC_HAS_PKGVERSION_BUGURL=y
-CT_CC_GCC_HAS_BUILD_ID=y
-CT_CC_GCC_HAS_LNK_HASH_STYLE=y
-CT_CC_GCC_USE_GMP_MPFR=y
-CT_CC_GCC_USE_MPC=y
-CT_CC_GCC_HAS_LIBQUADMATH=y
-CT_CC_GCC_HAS_LIBSANITIZER=y
-CT_CC_GCC_VERSION="linaro-4.8-2015.06"
-# CT_CC_LANG_FORTRAN is not set
+CT_CC_CHOICE_KSYM="GCC"
+CT_CC_GCC_SHOW=y
+
+#
+# Options for gcc
+#
+CT_CC_GCC_PKG_KSYM="GCC"
+CT_GCC_DIR_NAME="gcc"
+CT_GCC_USE_GNU=y
+# CT_GCC_USE_ORACLE is not set
+CT_GCC_USE="GCC"
+CT_GCC_PKG_NAME="gcc"
+CT_GCC_SRC_RELEASE=y
+# CT_GCC_SRC_DEVEL is not set
+CT_GCC_PATCH_ORDER="global"
+CT_GCC_V_13=y
+# CT_GCC_V_12 is not set
+# CT_GCC_V_11 is not set
+# CT_GCC_V_10 is not set
+# CT_GCC_V_9 is not set
+# CT_GCC_V_8 is not set
+# CT_GCC_V_7 is not set
+# CT_GCC_V_6 is not set
+# CT_GCC_V_5 is not set
+# CT_GCC_V_4_9 is not set
+CT_GCC_VERSION="13.2.0"
+CT_GCC_MIRRORS="$(CT_Mirrors GNU gcc/gcc-${CT_GCC_VERSION}) $(CT_Mirrors sourceware gcc/releases/gcc-${CT_GCC_VERSION})"
+CT_GCC_ARCHIVE_FILENAME="@{pkg_name}-@{version}"
+CT_GCC_ARCHIVE_DIRNAME="@{pkg_name}-@{version}"
+CT_GCC_ARCHIVE_FORMATS=".tar.xz .tar.gz"
+CT_GCC_SIGNATURE_FORMAT=""
+CT_GCC_later_than_13=y
+CT_GCC_13_or_later=y
+CT_GCC_later_than_12=y
+CT_GCC_12_or_later=y
+CT_GCC_later_than_11=y
+CT_GCC_11_or_later=y
+CT_GCC_later_than_10=y
+CT_GCC_10_or_later=y
+CT_GCC_later_than_9=y
+CT_GCC_9_or_later=y
+CT_GCC_later_than_8=y
+CT_GCC_8_or_later=y
+CT_GCC_later_than_7=y
+CT_GCC_7_or_later=y
+CT_GCC_later_than_6=y
+CT_GCC_6_or_later=y
+CT_GCC_later_than_5=y
+CT_GCC_5_or_later=y
+CT_GCC_later_than_4_9=y
+CT_GCC_4_9_or_later=y
+CT_GCC_REQUIRE_4_9_or_later=y
+CT_CC_GCC_HAS_LIBMPX=y
 CT_CC_GCC_ENABLE_CXX_FLAGS=""
 CT_CC_GCC_CORE_EXTRA_CONFIG_ARRAY=""
 CT_CC_GCC_EXTRA_CONFIG_ARRAY=""
-CT_CC_GCC_EXTRA_ENV_ARRAY=""
 CT_CC_GCC_STATIC_LIBSTDCXX=y
 # CT_CC_GCC_SYSTEM_ZLIB is not set
+CT_CC_GCC_CONFIG_TLS=m
 
 #
 # Optimisation features
 #
+CT_CC_GCC_USE_GRAPHITE=y
 
 #
 # Settings for libraries running on target
 #
+# CT_CC_GCC_ENABLE_DEFAULT_PIE is not set
 CT_CC_GCC_ENABLE_TARGET_OPTSPACE=y
 # CT_CC_GCC_LIBMUDFLAP is not set
 # CT_CC_GCC_LIBGOMP is not set
 # CT_CC_GCC_LIBSSP is not set
 # CT_CC_GCC_LIBQUADMATH is not set
 # CT_CC_GCC_LIBSANITIZER is not set
+CT_CC_GCC_LIBSTDCXX_VERBOSE=m
 
 #
 # Misc. obscure options.
 #
 CT_CC_CXA_ATEXIT=y
-CT_CC_GCC_DISABLE_PCH=y
+CT_CC_GCC_TM_CLONE_REGISTRY=m
+# CT_CC_GCC_DISABLE_PCH is not set
 CT_CC_GCC_SJLJ_EXCEPTIONS=m
 CT_CC_GCC_LDBL_128=m
 # CT_CC_GCC_BUILD_ID is not set
@@ -464,105 +635,243 @@ CT_CC_GCC_LNK_HASH_STYLE_DEFAULT=y
 # CT_CC_GCC_LNK_HASH_STYLE_GNU is not set
 # CT_CC_GCC_LNK_HASH_STYLE_BOTH is not set
 CT_CC_GCC_LNK_HASH_STYLE=""
-CT_CC_GCC_DEC_FLOAT_AUTO=y
-# CT_CC_GCC_DEC_FLOAT_BID is not set
-# CT_CC_GCC_DEC_FLOAT_DPD is not set
+CT_CC_GCC_DEC_FLOATS_AUTO=y
+# CT_CC_GCC_DEC_FLOATS_BID is not set
+# CT_CC_GCC_DEC_FLOATS_DPD is not set
 # CT_CC_GCC_DEC_FLOATS_NO is not set
-CT_CC_SUPPORT_CXX=y
-CT_CC_SUPPORT_FORTRAN=y
-CT_CC_SUPPORT_JAVA=y
-CT_CC_SUPPORT_ADA=y
-CT_CC_SUPPORT_OBJC=y
-CT_CC_SUPPORT_OBJCXX=y
-CT_CC_SUPPORT_GOLANG=y
+CT_CC_GCC_DEC_FLOATS=""
+CT_ALL_CC_CHOICES="GCC"
 
 #
 # Additional supported languages:
 #
 CT_CC_LANG_CXX=y
-# CT_CC_LANG_JAVA is not set
+# CT_CC_LANG_FORTRAN is not set
+# end of C compiler
 
 #
 # Debug facilities
 #
-# CT_DEBUG_dmalloc is not set
-# CT_DEBUG_duma is not set
-# CT_DEBUG_gdb is not set
-# CT_DEBUG_ltrace is not set
-# CT_DEBUG_strace is not set
+# CT_DEBUG_DUMA is not set
+# CT_DEBUG_GDB is not set
+# CT_DEBUG_LTRACE is not set
+# CT_DEBUG_STRACE is not set
+CT_ALL_DEBUG_CHOICES="DUMA GDB LTRACE STRACE"
+# end of Debug facilities
 
 #
 # Companion libraries
 #
-CT_COMPLIBS_NEEDED=y
+# CT_COMPLIBS_CHECK is not set
+# CT_COMP_LIBS_CLOOG is not set
+# CT_COMP_LIBS_EXPAT is not set
+CT_COMP_LIBS_GETTEXT=y
+CT_COMP_LIBS_GETTEXT_PKG_KSYM="GETTEXT"
+CT_GETTEXT_DIR_NAME="gettext"
+CT_GETTEXT_PKG_NAME="gettext"
+CT_GETTEXT_SRC_RELEASE=y
+# CT_GETTEXT_SRC_DEVEL is not set
+CT_GETTEXT_PATCH_ORDER="global"
+CT_GETTEXT_V_0_21=y
+# CT_GETTEXT_V_0_20_1 is not set
+# CT_GETTEXT_V_0_19_8_1 is not set
+CT_GETTEXT_VERSION="0.21"
+CT_GETTEXT_MIRRORS="$(CT_Mirrors GNU gettext)"
+CT_GETTEXT_ARCHIVE_FILENAME="@{pkg_name}-@{version}"
+CT_GETTEXT_ARCHIVE_DIRNAME="@{pkg_name}-@{version}"
+CT_GETTEXT_ARCHIVE_FORMATS=".tar.xz .tar.gz"
+CT_GETTEXT_SIGNATURE_FORMAT="packed/.sig"
+CT_GETTEXT_0_21_or_later=y
+CT_GETTEXT_0_21_or_older=y
+CT_GETTEXT_INCOMPATIBLE_WITH_UCLIBC_NG=y
+
+#
+# This version of gettext is not compatible with uClibc-NG. Select
+#
+
+#
+# a different version if uClibc-NG is used on the target or (in a
+#
+
+#
+# Canadian cross build) on the host.
+#
+CT_COMP_LIBS_GMP=y
+CT_COMP_LIBS_GMP_PKG_KSYM="GMP"
+CT_GMP_DIR_NAME="gmp"
+CT_GMP_PKG_NAME="gmp"
+CT_GMP_SRC_RELEASE=y
+# CT_GMP_SRC_DEVEL is not set
+CT_GMP_PATCH_ORDER="global"
+CT_GMP_V_6_2=y
+# CT_GMP_V_6_1 is not set
+CT_GMP_VERSION="6.2.1"
+CT_GMP_MIRRORS="https://gmplib.org/download/gmp https://gmplib.org/download/gmp/archive $(CT_Mirrors GNU gmp)"
+CT_GMP_ARCHIVE_FILENAME="@{pkg_name}-@{version}"
+CT_GMP_ARCHIVE_DIRNAME="@{pkg_name}-@{version}"
+CT_GMP_ARCHIVE_FORMATS=".tar.xz .tar.lz .tar.bz2"
+CT_GMP_SIGNATURE_FORMAT="packed/.sig"
+CT_COMP_LIBS_ISL=y
+CT_COMP_LIBS_ISL_PKG_KSYM="ISL"
+CT_ISL_DIR_NAME="isl"
+CT_ISL_PKG_NAME="isl"
+CT_ISL_SRC_RELEASE=y
+# CT_ISL_SRC_DEVEL is not set
+CT_ISL_PATCH_ORDER="global"
+CT_ISL_V_0_26=y
+# CT_ISL_V_0_25 is not set
+# CT_ISL_V_0_24 is not set
+# CT_ISL_V_0_23 is not set
+# CT_ISL_V_0_22 is not set
+# CT_ISL_V_0_21 is not set
+# CT_ISL_V_0_20 is not set
+# CT_ISL_V_0_19 is not set
+# CT_ISL_V_0_18 is not set
+# CT_ISL_V_0_17 is not set
+# CT_ISL_V_0_16 is not set
+# CT_ISL_V_0_15 is not set
+CT_ISL_VERSION="0.26"
+CT_ISL_MIRRORS="https://libisl.sourceforge.io"
+CT_ISL_ARCHIVE_FILENAME="@{pkg_name}-@{version}"
+CT_ISL_ARCHIVE_DIRNAME="@{pkg_name}-@{version}"
+CT_ISL_ARCHIVE_FORMATS=".tar.xz .tar.bz2 .tar.gz"
+CT_ISL_SIGNATURE_FORMAT=""
+CT_ISL_later_than_0_18=y
+CT_ISL_0_18_or_later=y
+CT_ISL_later_than_0_15=y
+CT_ISL_0_15_or_later=y
+# CT_COMP_LIBS_LIBELF is not set
+CT_COMP_LIBS_LIBICONV=y
+CT_COMP_LIBS_LIBICONV_PKG_KSYM="LIBICONV"
+CT_LIBICONV_DIR_NAME="libiconv"
+CT_LIBICONV_PKG_NAME="libiconv"
+CT_LIBICONV_SRC_RELEASE=y
+# CT_LIBICONV_SRC_DEVEL is not set
+CT_LIBICONV_PATCH_ORDER="global"
+CT_LIBICONV_V_1_16=y
+# CT_LIBICONV_V_1_15 is not set
+CT_LIBICONV_VERSION="1.16"
+CT_LIBICONV_MIRRORS="$(CT_Mirrors GNU libiconv)"
+CT_LIBICONV_ARCHIVE_FILENAME="@{pkg_name}-@{version}"
+CT_LIBICONV_ARCHIVE_DIRNAME="@{pkg_name}-@{version}"
+CT_LIBICONV_ARCHIVE_FORMATS=".tar.gz"
+CT_LIBICONV_SIGNATURE_FORMAT="packed/.sig"
+CT_COMP_LIBS_MPC=y
+CT_COMP_LIBS_MPC_PKG_KSYM="MPC"
+CT_MPC_DIR_NAME="mpc"
+CT_MPC_PKG_NAME="mpc"
+CT_MPC_SRC_RELEASE=y
+# CT_MPC_SRC_DEVEL is not set
+CT_MPC_PATCH_ORDER="global"
+CT_MPC_V_1_2=y
+CT_MPC_VERSION="1.2.1"
+CT_MPC_MIRRORS="https://www.multiprecision.org/downloads $(CT_Mirrors GNU mpc)"
+CT_MPC_ARCHIVE_FILENAME="@{pkg_name}-@{version}"
+CT_MPC_ARCHIVE_DIRNAME="@{pkg_name}-@{version}"
+CT_MPC_ARCHIVE_FORMATS=".tar.gz"
+CT_MPC_SIGNATURE_FORMAT="packed/.sig"
+CT_MPC_later_than_1_1_0=y
+CT_MPC_1_1_0_or_later=y
+CT_COMP_LIBS_MPFR=y
+CT_COMP_LIBS_MPFR_PKG_KSYM="MPFR"
+CT_MPFR_DIR_NAME="mpfr"
+CT_MPFR_PKG_NAME="mpfr"
+CT_MPFR_SRC_RELEASE=y
+# CT_MPFR_SRC_DEVEL is not set
+CT_MPFR_PATCH_ORDER="global"
+CT_MPFR_V_4_2=y
+CT_MPFR_VERSION="4.2.1"
+CT_MPFR_MIRRORS="https://www.mpfr.org/mpfr-${CT_MPFR_VERSION} $(CT_Mirrors GNU mpfr)"
+CT_MPFR_ARCHIVE_FILENAME="@{pkg_name}-@{version}"
+CT_MPFR_ARCHIVE_DIRNAME="@{pkg_name}-@{version}"
+CT_MPFR_ARCHIVE_FORMATS=".tar.xz .tar.bz2 .tar.gz .zip"
+CT_MPFR_SIGNATURE_FORMAT="packed/.asc"
+CT_MPFR_later_than_4_0_0=y
+CT_MPFR_4_0_0_or_later=y
+CT_COMP_LIBS_NCURSES=y
+CT_COMP_LIBS_NCURSES_PKG_KSYM="NCURSES"
+CT_NCURSES_DIR_NAME="ncurses"
+CT_NCURSES_PKG_NAME="ncurses"
+CT_NCURSES_SRC_RELEASE=y
+# CT_NCURSES_SRC_DEVEL is not set
+CT_NCURSES_PATCH_ORDER="global"
+CT_NCURSES_V_6_4=y
+# CT_NCURSES_V_6_2 is not set
+# CT_NCURSES_V_6_1 is not set
+# CT_NCURSES_V_6_0 is not set
+CT_NCURSES_VERSION="6.4"
+CT_NCURSES_MIRRORS="https://invisible-mirror.net/archives/ncurses $(CT_Mirrors GNU ncurses)"
+CT_NCURSES_ARCHIVE_FILENAME="@{pkg_name}-@{version}"
+CT_NCURSES_ARCHIVE_DIRNAME="@{pkg_name}-@{version}"
+CT_NCURSES_ARCHIVE_FORMATS=".tar.gz"
+CT_NCURSES_SIGNATURE_FORMAT="packed/.sig"
+CT_NCURSES_NEW_ABI=y
+CT_NCURSES_HOST_CONFIG_ARGS=""
+CT_NCURSES_HOST_DISABLE_DB=y
+CT_NCURSES_HOST_FALLBACKS="linux,xterm,xterm-color,xterm-256color,vt100"
+CT_NCURSES_TARGET_CONFIG_ARGS=""
+# CT_NCURSES_TARGET_DISABLE_DB is not set
+CT_NCURSES_TARGET_FALLBACKS=""
+CT_COMP_LIBS_ZLIB=y
+CT_COMP_LIBS_ZLIB_PKG_KSYM="ZLIB"
+CT_ZLIB_DIR_NAME="zlib"
+CT_ZLIB_PKG_NAME="zlib"
+CT_ZLIB_SRC_RELEASE=y
+# CT_ZLIB_SRC_DEVEL is not set
+CT_ZLIB_PATCH_ORDER="global"
+CT_ZLIB_V_1_2_13=y
+CT_ZLIB_VERSION="1.2.13"
+CT_ZLIB_MIRRORS="[REDACTED_PROFILE_URL]/zlib/releases/download/v${CT_ZLIB_VERSION} https://www.zlib.net/"
+CT_ZLIB_ARCHIVE_FILENAME="@{pkg_name}-@{version}"
+CT_ZLIB_ARCHIVE_DIRNAME="@{pkg_name}-@{version}"
+CT_ZLIB_ARCHIVE_FORMATS=".tar.xz .tar.gz"
+CT_ZLIB_SIGNATURE_FORMAT="packed/.asc"
+CT_COMP_LIBS_ZSTD=y
+CT_COMP_LIBS_ZSTD_PKG_KSYM="ZSTD"
+CT_ZSTD_DIR_NAME="zstd"
+CT_ZSTD_PKG_NAME="zstd"
+CT_ZSTD_SRC_RELEASE=y
+# CT_ZSTD_SRC_DEVEL is not set
+CT_ZSTD_PATCH_ORDER="global"
+CT_ZSTD_V_1_5_5=y
+# CT_ZSTD_V_1_5_2 is not set
+CT_ZSTD_VERSION="1.5.5"
+CT_ZSTD_MIRRORS="[REDACTED_PROFILE_URL]/zstd/releases/download/v${CT_ZSTD_VERSION} https://www.zstd.net/"
+CT_ZSTD_ARCHIVE_FILENAME="@{pkg_name}-@{version}"
+CT_ZSTD_ARCHIVE_DIRNAME="@{pkg_name}-@{version}"
+CT_ZSTD_ARCHIVE_FORMATS=".tar.gz"
+CT_ZSTD_SIGNATURE_FORMAT="packed/.sig"
+CT_ALL_COMP_LIBS_CHOICES="CLOOG EXPAT GETTEXT GMP GNUPRUMCU ISL LIBELF LIBICONV MPC MPFR NCURSES NEWLIB_NANO PICOLIBC ZLIB ZSTD"
 CT_LIBICONV_NEEDED=y
 CT_GETTEXT_NEEDED=y
 CT_GMP_NEEDED=y
 CT_MPFR_NEEDED=y
 CT_ISL_NEEDED=y
-CT_CLOOG_NEEDED=y
 CT_MPC_NEEDED=y
-CT_COMPLIBS=y
+CT_NCURSES_NEEDED=y
+CT_ZLIB_NEEDED=y
+CT_ZSTD_NEEDED=y
 CT_LIBICONV=y
 CT_GETTEXT=y
 CT_GMP=y
 CT_MPFR=y
 CT_ISL=y
-CT_CLOOG=y
 CT_MPC=y
-CT_LIBICONV_V_1_14=y
-CT_LIBICONV_VERSION="1.14"
-CT_GETTEXT_V_0_19_6=y
-CT_GETTEXT_VERSION="0.19.6"
-CT_GMP_V_6_0_0=y
-# CT_GMP_V_5_1_3 is not set
-# CT_GMP_V_5_1_1 is not set
-# CT_GMP_V_5_0_2 is not set
-# CT_GMP_V_5_0_1 is not set
-# CT_GMP_V_4_3_2 is not set
-# CT_GMP_V_4_3_1 is not set
-# CT_GMP_V_4_3_0 is not set
-CT_GMP_5_0_2_or_later=y
-CT_GMP_VERSION="6.0.0a"
-CT_MPFR_V_3_1_3=y
-# CT_MPFR_V_3_1_2 is not set
-# CT_MPFR_V_3_1_0 is not set
-# CT_MPFR_V_3_0_1 is not set
-# CT_MPFR_V_3_0_0 is not set
-# CT_MPFR_V_2_4_2 is not set
-# CT_MPFR_V_2_4_1 is not set
-# CT_MPFR_V_2_4_0 is not set
-CT_MPFR_VERSION="3.1.3"
-CT_ISL_V_0_14=y
-CT_ISL_V_0_14_or_later=y
-CT_ISL_V_0_12_or_later=y
-CT_ISL_VERSION="0.14"
-CT_CLOOG_V_0_18_4=y
-# CT_CLOOG_V_0_18_1 is not set
-# CT_CLOOG_V_0_18_0 is not set
-CT_CLOOG_VERSION="0.18.4"
-CT_CLOOG_0_18_4_or_later=y
-CT_CLOOG_0_18_or_later=y
-CT_MPC_V_1_0_3=y
-# CT_MPC_V_1_0_2 is not set
-# CT_MPC_V_1_0_1 is not set
-# CT_MPC_V_1_0 is not set
-# CT_MPC_V_0_9 is not set
-# CT_MPC_V_0_8_2 is not set
-# CT_MPC_V_0_8_1 is not set
-# CT_MPC_V_0_7 is not set
-CT_MPC_VERSION="1.0.3"
-
-#
-# Companion libraries common options
-#
-# CT_COMPLIBS_CHECK is not set
+CT_NCURSES=y
+CT_ZLIB=y
+CT_ZSTD=y
+# end of Companion libraries
 
 #
 # Companion tools
 #
-
-#
-# READ HELP before you say 'Y' below !!!
-#
-# CT_COMP_TOOLS is not set
+# CT_COMP_TOOLS_FOR_HOST is not set
+# CT_COMP_TOOLS_AUTOCONF is not set
+# CT_COMP_TOOLS_AUTOMAKE is not set
+# CT_COMP_TOOLS_BISON is not set
+# CT_COMP_TOOLS_DTC is not set
+# CT_COMP_TOOLS_LIBTOOL is not set
+# CT_COMP_TOOLS_M4 is not set
+# CT_COMP_TOOLS_MAKE is not set
+CT_ALL_COMP_TOOLS_CHOICES="AUTOCONF AUTOMAKE BISON DTC LIBTOOL M4 MAKE"
+# end of Companion tools
diff --git a/Engine/Build/BatchFiles/Linux/Toolchain/aarch64-unknown-linux-gnueabi.windows.config b/Engine/Build/BatchFiles/Linux/Toolchain/aarch64-unknown-linux-gnueabi.windows.config
--- a/Engine/Build/BatchFiles/Linux/Toolchain/aarch64-unknown-linux-gnueabi.windows.config
+++ b/Engine/Build/BatchFiles/Linux/Toolchain/aarch64-unknown-linux-gnueabi.windows.config
@@ -1,10 +1,34 @@
 #
 # Automatically generated file; DO NOT EDIT.
-# Crosstool-NG Configuration
-#
-CT_CONFIGURE_has_make381=y
-CT_CONFIGURE_has_xz=y
-CT_CONFIGURE_has_cvs=y
+# crosstool-NG 1.26.0 Configuration
+#
+CT_CONFIGURE_has_static_link=y
+CT_CONFIGURE_has_cxx11=y
+CT_CONFIGURE_has_wget=y
+CT_CONFIGURE_has_curl=y
+CT_CONFIGURE_has_rsync=y
+CT_CONFIGURE_has_make_3_81_or_newer=y
+CT_CONFIGURE_has_make_4_0_or_newer=y
+CT_CONFIGURE_has_libtool_2_4_or_newer=y
+CT_CONFIGURE_has_libtoolize_2_4_or_newer=y
+CT_CONFIGURE_has_autoconf_2_65_or_newer=y
+CT_CONFIGURE_has_autoreconf_2_65_or_newer=y
+CT_CONFIGURE_has_automake_1_15_or_newer=y
+CT_CONFIGURE_has_gnu_m4_1_4_12_or_newer=y
+CT_CONFIGURE_has_bison_2_7_or_newer=y
+CT_CONFIGURE_has_bison_3_0_4_or_newer=y
+CT_CONFIGURE_has_python=y
+CT_CONFIGURE_has_git=y
+CT_CONFIGURE_has_md5sum=y
+CT_CONFIGURE_has_sha1sum=y
+CT_CONFIGURE_has_sha256sum=y
+CT_CONFIGURE_has_sha512sum=y
+CT_CONFIGURE_has_install_with_strip_program=y
+CT_VERSION="1.26.0"
+CT_VCHECK=""
+CT_CONFIG_VERSION_ENV="4"
+CT_CONFIG_VERSION_CURRENT="4"
+CT_CONFIG_VERSION="4"
 CT_MODULES=y
 
 #
@@ -21,41 +45,49 @@ CT_MODULES=y
 #
 # Paths
 #
-CT_LOCAL_TARBALLS_DIR=""
+CT_LOCAL_TARBALLS_DIR="${HOME}/src"
+CT_SAVE_TARBALLS=y
+# CT_TARBALLS_BUILDROOT_LAYOUT is not set
 CT_WORK_DIR="${CT_TOP_DIR}/.build"
+CT_BUILD_TOP_DIR="${CT_WORK_DIR:-${CT_TOP_DIR}/.build}/${CT_HOST:+HOST-${CT_HOST}/}${CT_TARGET}"
+CT_BUILD_DIR="${CT_BUILD_TOP_DIR}/build"
 CT_PREFIX_DIR="${CT_TOP_DIR}/../OUTPUT-windows/${CT_TARGET}"
-CT_INSTALL_DIR="${CT_PREFIX_DIR}"
 CT_RM_RF_PREFIX_DIR=y
 CT_REMOVE_DOCS=y
-CT_INSTALL_DIR_RO=y
+CT_INSTALL_LICENSES=y
+CT_PREFIX_DIR_RO=y
 CT_STRIP_HOST_TOOLCHAIN_EXECUTABLES=y
 # CT_STRIP_TARGET_TOOLCHAIN_EXECUTABLES is not set
 
 #
 # Downloading
 #
+CT_DOWNLOAD_AGENT_WGET=y
+# CT_DOWNLOAD_AGENT_CURL is not set
+# CT_DOWNLOAD_AGENT_NONE is not set
 # CT_FORBID_DOWNLOAD is not set
 # CT_FORCE_DOWNLOAD is not set
 CT_CONNECT_TIMEOUT=10
+CT_DOWNLOAD_WGET_OPTIONS="--passive-ftp --tries=3 -nc --progress=dot:binary"
 # CT_ONLY_DOWNLOAD is not set
 # CT_USE_MIRROR is not set
+CT_VERIFY_DOWNLOAD_DIGEST=y
+CT_VERIFY_DOWNLOAD_DIGEST_SHA512=y
+# CT_VERIFY_DOWNLOAD_DIGEST_SHA256 is not set
+# CT_VERIFY_DOWNLOAD_DIGEST_SHA1 is not set
+# CT_VERIFY_DOWNLOAD_DIGEST_MD5 is not set
+CT_VERIFY_DOWNLOAD_DIGEST_ALG="sha512"
+# CT_VERIFY_DOWNLOAD_SIGNATURE is not set
 
 #
 # Extracting
 #
 # CT_FORCE_EXTRACT is not set
-CT_OVERIDE_CONFIG_GUESS_SUB=y
+CT_OVERRIDE_CONFIG_GUESS_SUB=y
 # CT_ONLY_EXTRACT is not set
-# CT_PATCH_BUNDLED is not set
-# CT_PATCH_LOCAL is not set
-CT_PATCH_BUNDLED_LOCAL=y
-# CT_PATCH_LOCAL_BUNDLED is not set
-# CT_PATCH_BUNDLED_FALLBACK_LOCAL is not set
-# CT_PATCH_LOCAL_FALLBACK_BUNDLED is not set
-# CT_PATCH_NONE is not set
-CT_PATCH_ORDER="bundled,local"
-CT_PATCH_USE_LOCAL=y
-CT_LOCAL_PATCH_DIR="${CT_TOP_DIR}/../../patches"
+CT_PATCH_BUNDLED=y
+# CT_PATCH_BUNDLED_LOCAL is not set
+CT_PATCH_ORDER="bundled"
 
 #
 # Build behavior
@@ -64,8 +96,9 @@ CT_PARALLEL_JOBS=0
 CT_LOAD=""
 CT_USE_PIPES=y
 CT_EXTRA_CFLAGS_FOR_BUILD=""
+CT_EXTRA_CXXFLAGS_FOR_BUILD=""
 CT_EXTRA_LDFLAGS_FOR_BUILD=""
-CT_EXTRA_CFLAGS_FOR_HOST="-fgnu89-inline"
+CT_EXTRA_CFLAGS_FOR_HOST=""
 CT_EXTRA_LDFLAGS_FOR_HOST=""
 # CT_CONFIG_SHELL_SH is not set
 # CT_CONFIG_SHELL_ASH is not set
@@ -78,93 +111,97 @@ CT_CONFIG_SHELL="${bash}"
 #
 # CT_LOG_ERROR is not set
 # CT_LOG_WARN is not set
-CT_LOG_INFO=y
-# CT_LOG_EXTRA is not set
+# CT_LOG_INFO is not set
+CT_LOG_EXTRA=y
 # CT_LOG_ALL is not set
 # CT_LOG_DEBUG is not set
-CT_LOG_LEVEL_MAX="INFO"
+CT_LOG_LEVEL_MAX="EXTRA"
 # CT_LOG_SEE_TOOLS_WARN is not set
 CT_LOG_PROGRESS_BAR=y
 CT_LOG_TO_FILE=y
 CT_LOG_FILE_COMPRESS=y
+# end of Paths and misc options
 
 #
 # Target options
 #
+# CT_ARCH_ALPHA is not set
+# CT_ARCH_ARC is not set
+CT_ARCH_ARM=y
+# CT_ARCH_AVR is not set
+# CT_ARCH_BPF is not set
+# CT_ARCH_M68K is not set
+# CT_ARCH_MIPS is not set
+# CT_ARCH_NIOS2 is not set
+# CT_ARCH_POWERPC is not set
+# CT_ARCH_PRU is not set
+# CT_ARCH_S390 is not set
+# CT_ARCH_SH is not set
+# CT_ARCH_SPARC is not set
+# CT_ARCH_X86 is not set
+# CT_ARCH_XTENSA is not set
 CT_ARCH="arm"
-CT_ARCH_SUPPORTS_BOTH_MMU=y
-CT_ARCH_SUPPORTS_BOTH_ENDIAN=y
-CT_ARCH_SUPPORTS_32=y
-CT_ARCH_SUPPORTS_64=y
-CT_ARCH_SUPPORTS_WITH_ARCH=y
-CT_ARCH_SUPPORTS_WITH_CPU=y
-CT_ARCH_SUPPORTS_WITH_TUNE=y
-CT_ARCH_DEFAULT_HAS_MMU=y
-CT_ARCH_DEFAULT_LE=y
-CT_ARCH_DEFAULT_32=y
-CT_ARCH_ARCH="armv8-a"
+CT_ARCH_CHOICE_KSYM="ARM"
 CT_ARCH_CPU=""
 CT_ARCH_TUNE=""
-# CT_ARCH_BE is not set
-CT_ARCH_LE=y
-# CT_ARCH_32 is not set
-CT_ARCH_64=y
-CT_ARCH_BITNESS=64
-# CT_ARCH_FLOAT_HW is not set
-# CT_ARCH_FLOAT_SW is not set
-CT_TARGET_CFLAGS=""
-CT_TARGET_LDFLAGS=""
-# CT_ARCH_alpha is not set
-CT_ARCH_arm=y
-# CT_ARCH_avr is not set
-# CT_ARCH_m68k is not set
-# CT_ARCH_mips is not set
-# CT_ARCH_nios2 is not set
-# CT_ARCH_powerpc is not set
-# CT_ARCH_s390 is not set
-# CT_ARCH_sh is not set
-# CT_ARCH_sparc is not set
-# CT_ARCH_x86 is not set
-# CT_ARCH_xtensa is not set
-CT_ARCH_alpha_AVAILABLE=y
-CT_ARCH_arm_AVAILABLE=y
-CT_ARCH_avr_AVAILABLE=y
-CT_ARCH_m68k_AVAILABLE=y
-CT_ARCH_microblaze_AVAILABLE=y
-CT_ARCH_mips_AVAILABLE=y
-CT_ARCH_nios2_AVAILABLE=y
-CT_ARCH_powerpc_AVAILABLE=y
-CT_ARCH_s390_AVAILABLE=y
-CT_ARCH_sh_AVAILABLE=y
-CT_ARCH_sparc_AVAILABLE=y
-CT_ARCH_x86_AVAILABLE=y
-CT_ARCH_xtensa_AVAILABLE=y
+CT_ARCH_ARM_SHOW=y
+
+#
+# Options for arm
+#
+CT_ARCH_ARM_PKG_KSYM=""
+CT_ARCH_ARM_MODE="arm"
+CT_ARCH_ARM_MODE_ARM=y
+# CT_ARCH_ARM_MODE_THUMB is not set
+# CT_ARCH_ARM_INTERWORKING is not set
+CT_ARCH_ARM_EABI_FORCE=y
+CT_ARCH_ARM_EABI=y
+CT_ALL_ARCH_CHOICES="ALPHA ARC ARM AVR BPF C6X LOONGARCH M68K MICROBLAZE MIPS MOXIE MSP430 NIOS2 POWERPC PRU RISCV S390 SH SPARC X86 XTENSA"
 CT_ARCH_SUFFIX=""
+# CT_OMIT_TARGET_VENDOR is not set
 
 #
 # Generic target options
 #
 # CT_MULTILIB is not set
+CT_DEMULTILIB=y
+CT_ARCH_SUPPORTS_BOTH_MMU=y
+CT_ARCH_DEFAULT_HAS_MMU=y
 CT_ARCH_USE_MMU=y
+CT_ARCH_SUPPORTS_FLAT_FORMAT=y
+CT_ARCH_SUPPORTS_LIBSANITIZER=y
+CT_ARCH_SUPPORTS_EITHER_ENDIAN=y
+CT_ARCH_DEFAULT_LE=y
+# CT_ARCH_BE is not set
+CT_ARCH_LE=y
 CT_ARCH_ENDIAN="little"
+CT_ARCH_SUPPORTS_32=y
+CT_ARCH_SUPPORTS_64=y
+CT_ARCH_DEFAULT_32=y
+CT_ARCH_BITNESS=64
+# CT_ARCH_32 is not set
+CT_ARCH_64=y
 
 #
 # Target optimisations
 #
+CT_ARCH_SUPPORTS_WITH_ARCH=y
+CT_ARCH_SUPPORTS_WITH_CPU=y
+CT_ARCH_SUPPORTS_WITH_TUNE=y
+CT_ARCH_SUPPORTS_WITH_FLOAT=y
+CT_ARCH_SUPPORTS_WITH_FPU=y
+CT_ARCH_SUPPORTS_SOFTFP=y
 CT_ARCH_EXCLUSIVE_WITH_CPU=y
-# CT_ARCH_FLOAT_AUTO is not set
+CT_ARCH_ARCH="armv8-a"
+CT_ARCH_FPU=""
+CT_ARCH_FLOAT_AUTO=y
+# CT_ARCH_FLOAT_HW is not set
 # CT_ARCH_FLOAT_SOFTFP is not set
-CT_ARCH_FLOAT=""
-
-#
-# arm other options
-#
-CT_ARCH_ARM_MODE="arm"
-CT_ARCH_ARM_MODE_ARM=y
-# CT_ARCH_ARM_MODE_THUMB is not set
-# CT_ARCH_ARM_INTERWORKING is not set
-CT_ARCH_ARM_EABI_FORCE=y
-CT_ARCH_ARM_EABI=y
+# CT_ARCH_FLOAT_SW is not set
+CT_TARGET_CFLAGS=""
+CT_TARGET_LDFLAGS=""
+CT_ARCH_FLOAT="auto"
+# end of Target options
 
 #
 # Toolchain options
@@ -177,8 +214,8 @@ CT_FORCE_SYSROOT=y
 CT_USE_SYSROOT=y
 CT_SYSROOT_NAME="sysroot"
 CT_SYSROOT_DIR_PREFIX=""
-CT_WANTS_STATIC_LINK=y
 CT_STATIC_TOOLCHAIN=y
+CT_SHOW_CT_VERSION=y
 CT_TOOLCHAIN_PKGVERSION=""
 CT_TOOLCHAIN_BUGURL=""
 
@@ -214,132 +251,279 @@ CT_HOST_SUFFIX=""
 # Misc options
 #
 # CT_TOOLCHAIN_ENABLE_NLS is not set
+# end of Toolchain options
 
 #
 # Operating System
 #
 CT_KERNEL_SUPPORTS_SHARED_LIBS=y
+# CT_KERNEL_BARE_METAL is not set
+CT_KERNEL_LINUX=y
 CT_KERNEL="linux"
-CT_KERNEL_VERSION="3.10.93"
-# CT_KERNEL_bare_metal is not set
-CT_KERNEL_linux=y
-CT_KERNEL_bare_metal_AVAILABLE=y
-CT_KERNEL_linux_AVAILABLE=y
-# CT_KERNEL_V_4_3 is not set
-# CT_KERNEL_V_4_2 is not set
-# CT_KERNEL_V_4_1 is not set
-# CT_KERNEL_V_3_18 is not set
-# CT_KERNEL_V_3_14 is not set
-# CT_KERNEL_V_3_12 is not set
-CT_KERNEL_V_3_10=y
-# CT_KERNEL_V_3_4 is not set
-# CT_KERNEL_V_3_2 is not set
-# CT_KERNEL_V_2_6_32 is not set
-# CT_KERNEL_LINUX_CUSTOM is not set
-CT_KERNEL_windows_AVAILABLE=y
-
-#
-# Common kernel options
-#
-CT_SHARED_LIBS=y
-
-#
-# linux other options
-#
+CT_KERNEL_CHOICE_KSYM="LINUX"
+CT_KERNEL_LINUX_SHOW=y
+
+#
+# Options for linux
+#
+CT_KERNEL_LINUX_PKG_KSYM="LINUX"
+CT_LINUX_DIR_NAME="linux"
+CT_LINUX_USE_WWW_KERNEL_ORG=y
+# CT_LINUX_USE_ORACLE is not set
+CT_LINUX_USE="LINUX"
+CT_LINUX_PKG_NAME="linux"
+CT_LINUX_SRC_RELEASE=y
+# CT_LINUX_SRC_DEVEL is not set
+CT_LINUX_PATCH_ORDER="global"
+# CT_LINUX_V_6_4 is not set
+# CT_LINUX_V_6_3 is not set
+# CT_LINUX_V_6_2 is not set
+# CT_LINUX_V_6_1 is not set
+# CT_LINUX_V_6_0 is not set
+# CT_LINUX_V_5_19 is not set
+# CT_LINUX_V_5_18 is not set
+# CT_LINUX_V_5_17 is not set
+# CT_LINUX_V_5_16 is not set
+# CT_LINUX_V_5_15 is not set
+# CT_LINUX_V_5_14 is not set
+# CT_LINUX_V_5_13 is not set
+# CT_LINUX_V_5_12 is not set
+# CT_LINUX_V_5_11 is not set
+# CT_LINUX_V_5_10 is not set
+# CT_LINUX_V_5_9 is not set
+# CT_LINUX_V_5_8 is not set
+# CT_LINUX_V_5_7 is not set
+# CT_LINUX_V_5_4 is not set
+# CT_LINUX_V_5_3 is not set
+# CT_LINUX_V_5_2 is not set
+# CT_LINUX_V_5_1 is not set
+# CT_LINUX_V_5_0 is not set
+# CT_LINUX_V_4_20 is not set
+# CT_LINUX_V_4_19 is not set
+CT_LINUX_V_4_18=y
+# CT_LINUX_V_4_17 is not set
+# CT_LINUX_V_4_16 is not set
+# CT_LINUX_V_4_15 is not set
+# CT_LINUX_V_4_14 is not set
+# CT_LINUX_V_4_13 is not set
+# CT_LINUX_V_4_12 is not set
+# CT_LINUX_V_4_11 is not set
+# CT_LINUX_V_4_10 is not set
+# CT_LINUX_V_4_9 is not set
+# CT_LINUX_V_4_4 is not set
+# CT_LINUX_V_4_1 is not set
+# CT_LINUX_V_3_16 is not set
+# CT_LINUX_V_3_13 is not set
+# CT_LINUX_V_3_12 is not set
+# CT_LINUX_V_3_10 is not set
+# CT_LINUX_V_3_4 is not set
+# CT_LINUX_V_3_2 is not set
+CT_LINUX_VERSION="4.18.20"
+CT_LINUX_MIRRORS="$(CT_Mirrors kernel.org linux ${CT_LINUX_VERSION})"
+CT_LINUX_ARCHIVE_FILENAME="@{pkg_name}-@{version}"
+CT_LINUX_ARCHIVE_DIRNAME="@{pkg_name}-@{version}"
+CT_LINUX_ARCHIVE_FORMATS=".tar.xz .tar.gz"
+CT_LINUX_SIGNATURE_FORMAT="unpacked/.sign"
+CT_LINUX_5_19_or_older=y
+CT_LINUX_older_than_5_19=y
+CT_LINUX_5_12_or_older=y
+CT_LINUX_older_than_5_12=y
+CT_LINUX_5_5_or_older=y
+CT_LINUX_older_than_5_5=y
+CT_LINUX_5_3_or_older=y
+CT_LINUX_older_than_5_3=y
+CT_LINUX_later_than_4_8=y
+CT_LINUX_4_8_or_later=y
+CT_LINUX_later_than_3_7=y
+CT_LINUX_3_7_or_later=y
+CT_LINUX_later_than_3_2=y
+CT_LINUX_3_2_or_later=y
+CT_LINUX_REQUIRE_3_2_or_later=y
+CT_KERNEL_DEP_RSYNC=y
 CT_KERNEL_LINUX_VERBOSITY_0=y
 # CT_KERNEL_LINUX_VERBOSITY_1 is not set
 # CT_KERNEL_LINUX_VERBOSITY_2 is not set
 CT_KERNEL_LINUX_VERBOSE_LEVEL=0
 CT_KERNEL_LINUX_INSTALL_CHECK=y
+CT_ALL_KERNEL_CHOICES="BARE_METAL LINUX WINDOWS"
+
+#
+# Common kernel options
+#
+CT_SHARED_LIBS=y
+# end of Operating System
 
 #
 # Binary utilities
 #
 CT_ARCH_BINFMT_ELF=y
+CT_BINUTILS_BINUTILS=y
 CT_BINUTILS="binutils"
-CT_BINUTILS_binutils=y
+CT_BINUTILS_CHOICE_KSYM="BINUTILS"
+CT_BINUTILS_BINUTILS_SHOW=y
+
+#
+# Options for binutils
+#
+CT_BINUTILS_BINUTILS_PKG_KSYM="BINUTILS"
+CT_BINUTILS_DIR_NAME="binutils"
+CT_BINUTILS_USE_GNU=y
+# CT_BINUTILS_USE_ORACLE is not set
+CT_BINUTILS_USE="BINUTILS"
+CT_BINUTILS_PKG_NAME="binutils"
+CT_BINUTILS_SRC_RELEASE=y
+# CT_BINUTILS_SRC_DEVEL is not set
+CT_BINUTILS_PATCH_ORDER="global"
+CT_BINUTILS_V_2_40=y
+# CT_BINUTILS_V_2_39 is not set
+# CT_BINUTILS_V_2_38 is not set
+# CT_BINUTILS_V_2_37 is not set
+# CT_BINUTILS_V_2_36 is not set
+# CT_BINUTILS_V_2_35 is not set
+# CT_BINUTILS_V_2_34 is not set
+# CT_BINUTILS_V_2_33 is not set
+# CT_BINUTILS_V_2_32 is not set
+# CT_BINUTILS_V_2_31 is not set
+# CT_BINUTILS_V_2_30 is not set
+# CT_BINUTILS_V_2_29 is not set
+# CT_BINUTILS_V_2_28 is not set
+# CT_BINUTILS_V_2_27 is not set
+# CT_BINUTILS_V_2_26 is not set
+CT_BINUTILS_VERSION="2.40"
+CT_BINUTILS_MIRRORS="$(CT_Mirrors GNU binutils) $(CT_Mirrors sourceware binutils/releases)"
+CT_BINUTILS_ARCHIVE_FILENAME="@{pkg_name}-@{version}"
+CT_BINUTILS_ARCHIVE_DIRNAME="@{pkg_name}-@{version}"
+CT_BINUTILS_ARCHIVE_FORMATS=".tar.xz .tar.bz2 .tar.gz"
+CT_BINUTILS_SIGNATURE_FORMAT="packed/.sig"
+CT_BINUTILS_later_than_2_39=y
+CT_BINUTILS_2_39_or_later=y
+CT_BINUTILS_later_than_2_30=y
+CT_BINUTILS_2_30_or_later=y
+CT_BINUTILS_later_than_2_27=y
+CT_BINUTILS_2_27_or_later=y
+CT_BINUTILS_later_than_2_26=y
+CT_BINUTILS_2_26_or_later=y
 
 #
 # GNU binutils
 #
-# CT_CC_BINUTILS_SHOW_LINARO is not set
-CT_BINUTILS_V_2_25_1=y
-# CT_BINUTILS_V_2_25 is not set
-# CT_BINUTILS_V_2_24 is not set
-# CT_BINUTILS_V_2_23_2 is not set
-# CT_BINUTILS_V_2_23_1 is not set
-# CT_BINUTILS_V_2_22 is not set
-# CT_BINUTILS_V_2_21_53 is not set
-# CT_BINUTILS_V_2_21_1a is not set
-# CT_BINUTILS_V_2_20_1a is not set
-# CT_BINUTILS_V_2_19_1a is not set
-# CT_BINUTILS_V_2_18a is not set
-CT_BINUTILS_VERSION="2.25.1"
-CT_BINUTILS_2_25_1_or_later=y
-CT_BINUTILS_2_25_or_later=y
-CT_BINUTILS_2_24_or_later=y
-CT_BINUTILS_2_23_or_later=y
-CT_BINUTILS_2_22_or_later=y
-CT_BINUTILS_2_21_or_later=y
-CT_BINUTILS_2_20_or_later=y
-CT_BINUTILS_2_19_or_later=y
-CT_BINUTILS_2_18_or_later=y
-CT_BINUTILS_HAS_HASH_STYLE=y
-CT_BINUTILS_HAS_GOLD=y
 CT_BINUTILS_GOLD_SUPPORTS_ARCH=y
-CT_BINUTILS_GOLD_SUPPORT=y
-CT_BINUTILS_HAS_PLUGINS=y
-CT_BINUTILS_HAS_PKGVERSION_BUGURL=y
-CT_BINUTILS_FORCE_LD_BFD=y
+CT_BINUTILS_FORCE_LD_BFD_DEFAULT=y
 CT_BINUTILS_LINKER_LD=y
-# CT_BINUTILS_LINKER_LD_GOLD is not set
-# CT_BINUTILS_LINKER_GOLD_LD is not set
 CT_BINUTILS_LINKERS_LIST="ld"
 CT_BINUTILS_LINKER_DEFAULT="bfd"
-# CT_BINUTILS_PLUGINS is not set
+CT_BINUTILS_RELRO=m
+CT_BINUTILS_DETERMINISTIC_ARCHIVES=y
 CT_BINUTILS_EXTRA_CONFIG_ARRAY=""
 # CT_BINUTILS_FOR_TARGET is not set
-
-#
-# binutils other options
-#
+CT_ALL_BINUTILS_CHOICES="BINUTILS"
+# end of Binary utilities
 
 #
 # C-library
 #
+CT_LIBC_GLIBC=y
+# CT_LIBC_UCLIBC_NG is not set
 CT_LIBC="glibc"
-CT_LIBC_VERSION="2.17"
-CT_LIBC_glibc=y
-# CT_LIBC_musl is not set
-# CT_LIBC_uClibc is not set
-CT_LIBC_avr_libc_AVAILABLE=y
-CT_LIBC_glibc_AVAILABLE=y
+CT_LIBC_CHOICE_KSYM="GLIBC"
+CT_LIBC_GLIBC_SHOW=y
+
+#
+# Options for glibc
+#
+CT_LIBC_GLIBC_PKG_KSYM="GLIBC"
+CT_GLIBC_DIR_NAME="glibc"
+CT_GLIBC_USE_GNU=y
+# CT_GLIBC_USE_ORACLE is not set
+CT_GLIBC_USE="GLIBC"
+CT_GLIBC_PKG_NAME="glibc"
+CT_GLIBC_SRC_RELEASE=y
+# CT_GLIBC_SRC_DEVEL is not set
+CT_GLIBC_PATCH_ORDER="global"
+CT_GLIBC_V_2_28=y
+# CT_GLIBC_V_2_27 is not set
+# CT_GLIBC_V_2_26 is not set
+# CT_GLIBC_V_2_25 is not set
+# CT_GLIBC_V_2_24 is not set
+# CT_GLIBC_V_2_23 is not set
+# CT_GLIBC_V_2_19 is not set
+# CT_GLIBC_V_2_17 is not set
+CT_GLIBC_VERSION="2.28"
+CT_GLIBC_MIRRORS="$(CT_Mirrors GNU glibc)"
+CT_GLIBC_ARCHIVE_FILENAME="@{pkg_name}-@{version}"
+CT_GLIBC_ARCHIVE_DIRNAME="@{pkg_name}-@{version}"
+CT_GLIBC_ARCHIVE_FORMATS=".tar.xz .tar.bz2 .tar.gz"
+CT_GLIBC_SIGNATURE_FORMAT="packed/.sig"
+CT_GLIBC_2_38_or_older=y
+CT_GLIBC_older_than_2_38=y
+CT_GLIBC_2_37_or_older=y
+CT_GLIBC_older_than_2_37=y
+CT_GLIBC_2_36_or_older=y
+CT_GLIBC_older_than_2_36=y
+CT_GLIBC_2_34_or_older=y
+CT_GLIBC_older_than_2_34=y
+CT_GLIBC_2_32_or_older=y
+CT_GLIBC_older_than_2_32=y
+CT_GLIBC_2_31_or_older=y
+CT_GLIBC_older_than_2_31=y
+CT_GLIBC_2_30_or_older=y
+CT_GLIBC_older_than_2_30=y
+CT_GLIBC_2_29_or_older=y
+CT_GLIBC_older_than_2_29=y
+CT_GLIBC_REQUIRE_older_than_2_29=y
+CT_GLIBC_2_28_or_later=y
+CT_GLIBC_2_28_or_older=y
+CT_GLIBC_later_than_2_27=y
+CT_GLIBC_2_27_or_later=y
+CT_GLIBC_later_than_2_26=y
+CT_GLIBC_2_26_or_later=y
+CT_GLIBC_later_than_2_25=y
+CT_GLIBC_2_25_or_later=y
+CT_GLIBC_later_than_2_24=y
+CT_GLIBC_2_24_or_later=y
+CT_GLIBC_later_than_2_23=y
+CT_GLIBC_2_23_or_later=y
+CT_GLIBC_later_than_2_20=y
+CT_GLIBC_2_20_or_later=y
+CT_GLIBC_later_than_2_17=y
+CT_GLIBC_2_17_or_later=y
+CT_GLIBC_later_than_2_14=y
+CT_GLIBC_2_14_or_later=y
+CT_GLIBC_DEP_KERNEL_HEADERS_VERSION=y
+CT_GLIBC_DEP_BINUTILS=y
+CT_GLIBC_DEP_GCC=y
+CT_GLIBC_DEP_PYTHON=y
+
+#
+# glibc >=2.29 requires python >3.4
+#
+CT_GLIBC_SPARC_ALLOW_V7=y
 CT_THREADS="nptl"
-# CT_CC_GLIBC_SHOW_LINARO is not set
-# CT_LIBC_GLIBC_V_2_22 is not set
-# CT_LIBC_GLIBC_V_2_21 is not set
-# CT_LIBC_GLIBC_V_2_20 is not set
-# CT_LIBC_GLIBC_V_2_19 is not set
-# CT_LIBC_GLIBC_V_2_18 is not set
-CT_LIBC_GLIBC_V_2_17=y
-# CT_LIBC_GLIBC_V_2_16_0 is not set
-# CT_LIBC_GLIBC_V_2_15 is not set
-# CT_LIBC_GLIBC_V_2_14_1 is not set
-# CT_LIBC_GLIBC_V_2_14 is not set
-# CT_LIBC_GLIBC_V_2_13 is not set
-# CT_LIBC_GLIBC_V_2_12_2 is not set
-# CT_LIBC_GLIBC_V_2_12_1 is not set
-# CT_LIBC_GLIBC_V_2_11_1 is not set
-# CT_LIBC_GLIBC_V_2_11 is not set
-# CT_LIBC_GLIBC_V_2_10_1 is not set
-# CT_LIBC_GLIBC_V_2_9 is not set
-# CT_LIBC_GLIBC_V_2_8 is not set
-CT_LIBC_GLIBC_2_17_or_later=y
-CT_LIBC_mingw_AVAILABLE=y
-CT_LIBC_musl_AVAILABLE=y
-CT_LIBC_newlib_AVAILABLE=y
-CT_LIBC_none_AVAILABLE=y
-CT_LIBC_uClibc_AVAILABLE=y
+CT_GLIBC_BUILD_SSP=y
+CT_GLIBC_HAS_LIBIDN_ADDON=y
+# CT_GLIBC_USE_LIBIDN_ADDON is not set
+CT_GLIBC_NO_SPARC_V8=y
+CT_GLIBC_HAS_OBSOLETE_RPC=y
+CT_GLIBC_EXTRA_CONFIG_ARRAY=""
+CT_GLIBC_CONFIGPARMS=""
+CT_GLIBC_ENABLE_DEBUG=y
+CT_GLIBC_EXTRA_CFLAGS="-Wno-missing-attributes -Wno-array-bounds -Wno-array-parameter -Wno-stringop-overflow -Wno-maybe-uninitialized"
+CT_GLIBC_ENABLE_OBSOLETE_RPC=y
+# CT_GLIBC_DISABLE_VERSIONING is not set
+CT_GLIBC_OLDEST_ABI=""
+CT_GLIBC_FORCE_UNWIND=y
+# CT_GLIBC_LOCALES is not set
+# CT_GLIBC_KERNEL_VERSION_NONE is not set
+CT_GLIBC_KERNEL_VERSION_AS_HEADERS=y
+# CT_GLIBC_KERNEL_VERSION_CHOSEN is not set
+CT_GLIBC_MIN_KERNEL="4.18.20"
+CT_GLIBC_SSP_DEFAULT=y
+# CT_GLIBC_SSP_NO is not set
+# CT_GLIBC_SSP_YES is not set
+# CT_GLIBC_SSP_ALL is not set
+# CT_GLIBC_SSP_STRONG is not set
+CT_GLIBC_ENABLE_COMMON_FLAG=y
+CT_ALL_LIBC_CHOICES="AVR_LIBC GLIBC MINGW_W64 MOXIEBOX MUSL NEWLIB NONE PICOLIBC UCLIBC_NG"
 CT_LIBC_SUPPORT_THREADS_ANY=y
 CT_LIBC_SUPPORT_THREADS_NATIVE=y
 
@@ -347,121 +531,106 @@ CT_LIBC_SUPPORT_THREADS_NATIVE=y
 # Common C library options
 #
 CT_THREADS_NATIVE=y
+# CT_CREATE_LDSO_CONF is not set
 CT_LIBC_XLDD=y
-
-#
-# glibc other options
-#
-# CT_LIBC_GLIBC_PORTS_EXTERNAL is not set
-CT_LIBC_GLIBC_MAY_FORCE_PORTS=y
-CT_LIBC_glibc_familly=y
-CT_LIBC_GLIBC_EXTRA_CONFIG_ARRAY=""
-CT_LIBC_GLIBC_CONFIGPARMS=""
-CT_LIBC_GLIBC_EXTRA_CFLAGS=""
-CT_LIBC_EXTRA_CC_ARGS=""
-# CT_LIBC_DISABLE_VERSIONING is not set
-CT_LIBC_OLDEST_ABI=""
-CT_LIBC_GLIBC_FORCE_UNWIND=y
-CT_LIBC_GLIBC_USE_PORTS=y
-CT_LIBC_ADDONS_LIST=""
-
-#
-# WARNING !!!
-#
-
-#
-#   For glibc >= 2.8, it can happen that the tarballs
-#
-
-#
-#   for the addons are not available for download.
-#
-
-#
-#   If that happens, bad luck... Try a previous version
-#
-
-#
-#   or try again later... :-(
-#
-# CT_LIBC_LOCALES is not set
-# CT_LIBC_GLIBC_KERNEL_VERSION_NONE is not set
-CT_LIBC_GLIBC_KERNEL_VERSION_AS_HEADERS=y
-# CT_LIBC_GLIBC_KERNEL_VERSION_CHOSEN is not set
-CT_LIBC_GLIBC_MIN_KERNEL="3.10.93"
+# end of C-library
 
 #
 # C compiler
 #
+CT_CC_CORE_NEEDED=y
+CT_CC_SUPPORT_CXX=y
+CT_CC_SUPPORT_FORTRAN=y
+CT_CC_SUPPORT_ADA=y
+CT_CC_SUPPORT_D=y
+CT_CC_SUPPORT_JIT=y
+CT_CC_SUPPORT_OBJC=y
+CT_CC_SUPPORT_OBJCXX=y
+CT_CC_SUPPORT_GOLANG=y
+CT_CC_GCC=y
 CT_CC="gcc"
-CT_CC_CORE_PASSES_NEEDED=y
-CT_CC_CORE_PASS_1_NEEDED=y
-CT_CC_CORE_PASS_2_NEEDED=y
-CT_CC_gcc=y
-CT_CC_GCC_SHOW_LINARO=y
-# CT_CC_GCC_V_5_2_0 is not set
-# CT_CC_GCC_V_linaro_4_9=y
-# CT_CC_GCC_V_4_9_3 is not set
-CT_CC_GCC_V_linaro_4_8=y
-# CT_CC_GCC_V_4_8_5 is not set
-# CT_CC_GCC_V_linaro_4_7 is not set
-# CT_CC_GCC_V_4_7_4 is not set
-# CT_CC_GCC_V_linaro_4_6 is not set
-# CT_CC_GCC_V_4_6_4 is not set
-# CT_CC_GCC_V_linaro_4_5 is not set
-# CT_CC_GCC_V_4_5_4 is not set
-# CT_CC_GCC_V_linaro_4_4 is not set
-# CT_CC_GCC_V_4_4_7 is not set
-# CT_CC_GCC_V_4_3_6 is not set
-# CT_CC_GCC_V_4_2_4 is not set
-# CT_CC_GCC_CUSTOM is not set
-CT_CC_GCC_4_2_or_later=y
-CT_CC_GCC_4_3_or_later=y
-CT_CC_GCC_4_4_or_later=y
-CT_CC_GCC_4_5_or_later=y
-CT_CC_GCC_4_6_or_later=y
-CT_CC_GCC_4_7_or_later=y
-CT_CC_GCC_4_8_or_later=y
-CT_CC_GCC_4_8=y
-CT_CC_GCC_4_9_or_later=y
-CT_CC_GCC_HAS_GRAPHITE=y
-# CT_CC_GCC_USE_GRAPHITE is not set
-CT_CC_GCC_HAS_LTO=y
-# CT_CC_GCC_USE_LTO is not set
-CT_CC_GCC_HAS_PKGVERSION_BUGURL=y
-CT_CC_GCC_HAS_BUILD_ID=y
-CT_CC_GCC_HAS_LNK_HASH_STYLE=y
-CT_CC_GCC_USE_GMP_MPFR=y
-CT_CC_GCC_USE_MPC=y
-CT_CC_GCC_HAS_LIBQUADMATH=y
-CT_CC_GCC_HAS_LIBSANITIZER=y
-CT_CC_GCC_VERSION="linaro-4.8-2015.06"
-# CT_CC_LANG_FORTRAN is not set
+CT_CC_CHOICE_KSYM="GCC"
+CT_CC_GCC_SHOW=y
+
+#
+# Options for gcc
+#
+CT_CC_GCC_PKG_KSYM="GCC"
+CT_GCC_DIR_NAME="gcc"
+CT_GCC_USE_GNU=y
+# CT_GCC_USE_ORACLE is not set
+CT_GCC_USE="GCC"
+CT_GCC_PKG_NAME="gcc"
+CT_GCC_SRC_RELEASE=y
+# CT_GCC_SRC_DEVEL is not set
+CT_GCC_PATCH_ORDER="global"
+CT_GCC_V_13=y
+# CT_GCC_V_12 is not set
+# CT_GCC_V_11 is not set
+# CT_GCC_V_10 is not set
+# CT_GCC_V_9 is not set
+# CT_GCC_V_8 is not set
+# CT_GCC_V_7 is not set
+# CT_GCC_V_6 is not set
+# CT_GCC_V_5 is not set
+# CT_GCC_V_4_9 is not set
+CT_GCC_VERSION="13.2.0"
+CT_GCC_MIRRORS="$(CT_Mirrors GNU gcc/gcc-${CT_GCC_VERSION}) $(CT_Mirrors sourceware gcc/releases/gcc-${CT_GCC_VERSION})"
+CT_GCC_ARCHIVE_FILENAME="@{pkg_name}-@{version}"
+CT_GCC_ARCHIVE_DIRNAME="@{pkg_name}-@{version}"
+CT_GCC_ARCHIVE_FORMATS=".tar.xz .tar.gz"
+CT_GCC_SIGNATURE_FORMAT=""
+CT_GCC_later_than_13=y
+CT_GCC_13_or_later=y
+CT_GCC_later_than_12=y
+CT_GCC_12_or_later=y
+CT_GCC_later_than_11=y
+CT_GCC_11_or_later=y
+CT_GCC_later_than_10=y
+CT_GCC_10_or_later=y
+CT_GCC_later_than_9=y
+CT_GCC_9_or_later=y
+CT_GCC_later_than_8=y
+CT_GCC_8_or_later=y
+CT_GCC_later_than_7=y
+CT_GCC_7_or_later=y
+CT_GCC_later_than_6=y
+CT_GCC_6_or_later=y
+CT_GCC_later_than_5=y
+CT_GCC_5_or_later=y
+CT_GCC_later_than_4_9=y
+CT_GCC_4_9_or_later=y
+CT_GCC_REQUIRE_4_9_or_later=y
+CT_CC_GCC_HAS_LIBMPX=y
 CT_CC_GCC_ENABLE_CXX_FLAGS=""
 CT_CC_GCC_CORE_EXTRA_CONFIG_ARRAY=""
 CT_CC_GCC_EXTRA_CONFIG_ARRAY=""
-CT_CC_GCC_EXTRA_ENV_ARRAY=""
 CT_CC_GCC_STATIC_LIBSTDCXX=y
 # CT_CC_GCC_SYSTEM_ZLIB is not set
+CT_CC_GCC_CONFIG_TLS=m
 
 #
 # Optimisation features
 #
+# CT_CC_GCC_USE_GRAPHITE is not set
 
 #
 # Settings for libraries running on target
 #
+# CT_CC_GCC_ENABLE_DEFAULT_PIE is not set
 CT_CC_GCC_ENABLE_TARGET_OPTSPACE=y
 # CT_CC_GCC_LIBMUDFLAP is not set
 # CT_CC_GCC_LIBGOMP is not set
 # CT_CC_GCC_LIBSSP is not set
 # CT_CC_GCC_LIBQUADMATH is not set
 # CT_CC_GCC_LIBSANITIZER is not set
+CT_CC_GCC_LIBSTDCXX_VERBOSE=m
 
 #
 # Misc. obscure options.
 #
 CT_CC_CXA_ATEXIT=y
+CT_CC_GCC_TM_CLONE_REGISTRY=m
 CT_CC_GCC_DISABLE_PCH=y
 CT_CC_GCC_SJLJ_EXCEPTIONS=m
 CT_CC_GCC_LDBL_128=m
@@ -471,105 +640,199 @@ CT_CC_GCC_LNK_HASH_STYLE_DEFAULT=y
 # CT_CC_GCC_LNK_HASH_STYLE_GNU is not set
 # CT_CC_GCC_LNK_HASH_STYLE_BOTH is not set
 CT_CC_GCC_LNK_HASH_STYLE=""
-CT_CC_GCC_DEC_FLOAT_AUTO=y
-# CT_CC_GCC_DEC_FLOAT_BID is not set
-# CT_CC_GCC_DEC_FLOAT_DPD is not set
+CT_CC_GCC_DEC_FLOATS_AUTO=y
+# CT_CC_GCC_DEC_FLOATS_BID is not set
+# CT_CC_GCC_DEC_FLOATS_DPD is not set
 # CT_CC_GCC_DEC_FLOATS_NO is not set
-CT_CC_SUPPORT_CXX=y
-CT_CC_SUPPORT_FORTRAN=y
-CT_CC_SUPPORT_JAVA=y
-CT_CC_SUPPORT_ADA=y
-CT_CC_SUPPORT_OBJC=y
-CT_CC_SUPPORT_OBJCXX=y
-CT_CC_SUPPORT_GOLANG=y
+CT_CC_GCC_DEC_FLOATS=""
+CT_ALL_CC_CHOICES="GCC"
 
 #
 # Additional supported languages:
 #
 CT_CC_LANG_CXX=y
-# CT_CC_LANG_JAVA is not set
+# CT_CC_LANG_FORTRAN is not set
+# end of C compiler
 
 #
 # Debug facilities
 #
-# CT_DEBUG_dmalloc is not set
-# CT_DEBUG_duma is not set
-# CT_DEBUG_gdb is not set
-# CT_DEBUG_ltrace is not set
-# CT_DEBUG_strace is not set
+# CT_DEBUG_DUMA is not set
+# CT_DEBUG_GDB is not set
+# CT_DEBUG_LTRACE is not set
+# CT_DEBUG_STRACE is not set
+CT_ALL_DEBUG_CHOICES="DUMA GDB LTRACE STRACE"
+# end of Debug facilities
 
 #
 # Companion libraries
 #
-CT_COMPLIBS_NEEDED=y
+# CT_COMPLIBS_CHECK is not set
+# CT_COMP_LIBS_CLOOG is not set
+# CT_COMP_LIBS_EXPAT is not set
+CT_COMP_LIBS_GETTEXT=y
+CT_COMP_LIBS_GETTEXT_PKG_KSYM="GETTEXT"
+CT_GETTEXT_DIR_NAME="gettext"
+CT_GETTEXT_PKG_NAME="gettext"
+CT_GETTEXT_SRC_RELEASE=y
+# CT_GETTEXT_SRC_DEVEL is not set
+CT_GETTEXT_PATCH_ORDER="global"
+CT_GETTEXT_V_0_21=y
+# CT_GETTEXT_V_0_20_1 is not set
+# CT_GETTEXT_V_0_19_8_1 is not set
+CT_GETTEXT_VERSION="0.21"
+CT_GETTEXT_MIRRORS="$(CT_Mirrors GNU gettext)"
+CT_GETTEXT_ARCHIVE_FILENAME="@{pkg_name}-@{version}"
+CT_GETTEXT_ARCHIVE_DIRNAME="@{pkg_name}-@{version}"
+CT_GETTEXT_ARCHIVE_FORMATS=".tar.xz .tar.gz"
+CT_GETTEXT_SIGNATURE_FORMAT="packed/.sig"
+CT_GETTEXT_0_21_or_later=y
+CT_GETTEXT_0_21_or_older=y
+CT_GETTEXT_INCOMPATIBLE_WITH_UCLIBC_NG=y
+
+#
+# This version of gettext is not compatible with uClibc-NG. Select
+#
+
+#
+# a different version if uClibc-NG is used on the target or (in a
+#
+
+#
+# Canadian cross build) on the host.
+#
+CT_COMP_LIBS_GMP=y
+CT_COMP_LIBS_GMP_PKG_KSYM="GMP"
+CT_GMP_DIR_NAME="gmp"
+CT_GMP_PKG_NAME="gmp"
+CT_GMP_SRC_RELEASE=y
+# CT_GMP_SRC_DEVEL is not set
+CT_GMP_PATCH_ORDER="global"
+CT_GMP_V_6_2=y
+# CT_GMP_V_6_1 is not set
+CT_GMP_VERSION="6.2.1"
+CT_GMP_MIRRORS="https://gmplib.org/download/gmp https://gmplib.org/download/gmp/archive $(CT_Mirrors GNU gmp)"
+CT_GMP_ARCHIVE_FILENAME="@{pkg_name}-@{version}"
+CT_GMP_ARCHIVE_DIRNAME="@{pkg_name}-@{version}"
+CT_GMP_ARCHIVE_FORMATS=".tar.xz .tar.lz .tar.bz2"
+CT_GMP_SIGNATURE_FORMAT="packed/.sig"
+# CT_COMP_LIBS_ISL is not set
+# CT_COMP_LIBS_LIBELF is not set
+CT_COMP_LIBS_LIBICONV=y
+CT_COMP_LIBS_LIBICONV_PKG_KSYM="LIBICONV"
+CT_LIBICONV_DIR_NAME="libiconv"
+CT_LIBICONV_PKG_NAME="libiconv"
+CT_LIBICONV_SRC_RELEASE=y
+# CT_LIBICONV_SRC_DEVEL is not set
+CT_LIBICONV_PATCH_ORDER="global"
+CT_LIBICONV_V_1_16=y
+# CT_LIBICONV_V_1_15 is not set
+CT_LIBICONV_VERSION="1.16"
+CT_LIBICONV_MIRRORS="$(CT_Mirrors GNU libiconv)"
+CT_LIBICONV_ARCHIVE_FILENAME="@{pkg_name}-@{version}"
+CT_LIBICONV_ARCHIVE_DIRNAME="@{pkg_name}-@{version}"
+CT_LIBICONV_ARCHIVE_FORMATS=".tar.gz"
+CT_LIBICONV_SIGNATURE_FORMAT="packed/.sig"
+CT_COMP_LIBS_MPC=y
+CT_COMP_LIBS_MPC_PKG_KSYM="MPC"
+CT_MPC_DIR_NAME="mpc"
+CT_MPC_PKG_NAME="mpc"
+CT_MPC_SRC_RELEASE=y
+# CT_MPC_SRC_DEVEL is not set
+CT_MPC_PATCH_ORDER="global"
+CT_MPC_V_1_2=y
+CT_MPC_VERSION="1.2.1"
+CT_MPC_MIRRORS="https://www.multiprecision.org/downloads $(CT_Mirrors GNU mpc)"
+CT_MPC_ARCHIVE_FILENAME="@{pkg_name}-@{version}"
+CT_MPC_ARCHIVE_DIRNAME="@{pkg_name}-@{version}"
+CT_MPC_ARCHIVE_FORMATS=".tar.gz"
+CT_MPC_SIGNATURE_FORMAT="packed/.sig"
+CT_MPC_later_than_1_1_0=y
+CT_MPC_1_1_0_or_later=y
+CT_COMP_LIBS_MPFR=y
+CT_COMP_LIBS_MPFR_PKG_KSYM="MPFR"
+CT_MPFR_DIR_NAME="mpfr"
+CT_MPFR_PKG_NAME="mpfr"
+CT_MPFR_SRC_RELEASE=y
+# CT_MPFR_SRC_DEVEL is not set
+CT_MPFR_PATCH_ORDER="global"
+CT_MPFR_V_4_2=y
+CT_MPFR_VERSION="4.2.1"
+CT_MPFR_MIRRORS="https://www.mpfr.org/mpfr-${CT_MPFR_VERSION} $(CT_Mirrors GNU mpfr)"
+CT_MPFR_ARCHIVE_FILENAME="@{pkg_name}-@{version}"
+CT_MPFR_ARCHIVE_DIRNAME="@{pkg_name}-@{version}"
+CT_MPFR_ARCHIVE_FORMATS=".tar.xz .tar.bz2 .tar.gz .zip"
+CT_MPFR_SIGNATURE_FORMAT="packed/.asc"
+CT_MPFR_later_than_4_0_0=y
+CT_MPFR_4_0_0_or_later=y
+CT_COMP_LIBS_NCURSES=y
+CT_COMP_LIBS_NCURSES_PKG_KSYM="NCURSES"
+CT_NCURSES_DIR_NAME="ncurses"
+CT_NCURSES_PKG_NAME="ncurses"
+CT_NCURSES_SRC_RELEASE=y
+# CT_NCURSES_SRC_DEVEL is not set
+CT_NCURSES_PATCH_ORDER="global"
+CT_NCURSES_V_6_4=y
+# CT_NCURSES_V_6_2 is not set
+# CT_NCURSES_V_6_1 is not set
+# CT_NCURSES_V_6_0 is not set
+CT_NCURSES_VERSION="6.4"
+CT_NCURSES_MIRRORS="https://invisible-mirror.net/archives/ncurses $(CT_Mirrors GNU ncurses)"
+CT_NCURSES_ARCHIVE_FILENAME="@{pkg_name}-@{version}"
+CT_NCURSES_ARCHIVE_DIRNAME="@{pkg_name}-@{version}"
+CT_NCURSES_ARCHIVE_FORMATS=".tar.gz"
+CT_NCURSES_SIGNATURE_FORMAT="packed/.sig"
+CT_NCURSES_NEW_ABI=y
+CT_NCURSES_HOST_CONFIG_ARGS=""
+CT_NCURSES_HOST_DISABLE_DB=y
+CT_NCURSES_HOST_FALLBACKS="linux,xterm,xterm-color,xterm-256color,vt100"
+CT_NCURSES_TARGET_CONFIG_ARGS=""
+# CT_NCURSES_TARGET_DISABLE_DB is not set
+CT_NCURSES_TARGET_FALLBACKS=""
+# CT_COMP_LIBS_ZLIB is not set
+CT_COMP_LIBS_ZSTD=y
+CT_COMP_LIBS_ZSTD_PKG_KSYM="ZSTD"
+CT_ZSTD_DIR_NAME="zstd"
+CT_ZSTD_PKG_NAME="zstd"
+CT_ZSTD_SRC_RELEASE=y
+# CT_ZSTD_SRC_DEVEL is not set
+CT_ZSTD_PATCH_ORDER="global"
+CT_ZSTD_V_1_5_5=y
+# CT_ZSTD_V_1_5_2 is not set
+CT_ZSTD_VERSION="1.5.5"
+CT_ZSTD_MIRRORS="[REDACTED_PROFILE_URL]/zstd/releases/download/v${CT_ZSTD_VERSION} https://www.zstd.net/"
+CT_ZSTD_ARCHIVE_FILENAME="@{pkg_name}-@{version}"
+CT_ZSTD_ARCHIVE_DIRNAME="@{pkg_name}-@{version}"
+CT_ZSTD_ARCHIVE_FORMATS=".tar.gz"
+CT_ZSTD_SIGNATURE_FORMAT="packed/.sig"
+CT_ALL_COMP_LIBS_CHOICES="CLOOG EXPAT GETTEXT GMP GNUPRUMCU ISL LIBELF LIBICONV MPC MPFR NCURSES NEWLIB_NANO PICOLIBC ZLIB ZSTD"
 CT_LIBICONV_NEEDED=y
 CT_GETTEXT_NEEDED=y
 CT_GMP_NEEDED=y
 CT_MPFR_NEEDED=y
-CT_ISL_NEEDED=y
-CT_CLOOG_NEEDED=y
 CT_MPC_NEEDED=y
-CT_COMPLIBS=y
+CT_NCURSES_NEEDED=y
+# CT_ZLIB_NEEDED is not set
+CT_ZSTD_NEEDED=y
 CT_LIBICONV=y
 CT_GETTEXT=y
 CT_GMP=y
 CT_MPFR=y
-CT_ISL=y
-CT_CLOOG=y
 CT_MPC=y
-CT_LIBICONV_V_1_14=y
-CT_LIBICONV_VERSION="1.14"
-CT_GETTEXT_V_0_19_6=y
-CT_GETTEXT_VERSION="0.19.6"
-CT_GMP_V_6_0_0=y
-# CT_GMP_V_5_1_3 is not set
-# CT_GMP_V_5_1_1 is not set
-# CT_GMP_V_5_0_2 is not set
-# CT_GMP_V_5_0_1 is not set
-# CT_GMP_V_4_3_2 is not set
-# CT_GMP_V_4_3_1 is not set
-# CT_GMP_V_4_3_0 is not set
-CT_GMP_5_0_2_or_later=y
-CT_GMP_VERSION="6.0.0a"
-CT_MPFR_V_3_1_3=y
-# CT_MPFR_V_3_1_2 is not set
-# CT_MPFR_V_3_1_0 is not set
-# CT_MPFR_V_3_0_1 is not set
-# CT_MPFR_V_3_0_0 is not set
-# CT_MPFR_V_2_4_2 is not set
-# CT_MPFR_V_2_4_1 is not set
-# CT_MPFR_V_2_4_0 is not set
-CT_MPFR_VERSION="3.1.3"
-CT_ISL_V_0_14=y
-CT_ISL_V_0_14_or_later=y
-CT_ISL_V_0_12_or_later=y
-CT_ISL_VERSION="0.14"
-CT_CLOOG_V_0_18_4=y
-# CT_CLOOG_V_0_18_1 is not set
-# CT_CLOOG_V_0_18_0 is not set
-CT_CLOOG_VERSION="0.18.4"
-CT_CLOOG_0_18_4_or_later=y
-CT_CLOOG_0_18_or_later=y
-CT_MPC_V_1_0_3=y
-# CT_MPC_V_1_0_2 is not set
-# CT_MPC_V_1_0_1 is not set
-# CT_MPC_V_1_0 is not set
-# CT_MPC_V_0_9 is not set
-# CT_MPC_V_0_8_2 is not set
-# CT_MPC_V_0_8_1 is not set
-# CT_MPC_V_0_7 is not set
-CT_MPC_VERSION="1.0.3"
-
-#
-# Companion libraries common options
-#
-# CT_COMPLIBS_CHECK is not set
+CT_NCURSES=y
+CT_ZSTD=y
+# end of Companion libraries
 
 #
 # Companion tools
 #
-
-#
-# READ HELP before you say 'Y' below !!!
-#
-# CT_COMP_TOOLS is not set
+# CT_COMP_TOOLS_FOR_HOST is not set
+# CT_COMP_TOOLS_AUTOCONF is not set
+# CT_COMP_TOOLS_AUTOMAKE is not set
+# CT_COMP_TOOLS_BISON is not set
+# CT_COMP_TOOLS_DTC is not set
+# CT_COMP_TOOLS_LIBTOOL is not set
+# CT_COMP_TOOLS_M4 is not set
+# CT_COMP_TOOLS_MAKE is not set
+CT_ALL_COMP_TOOLS_CHOICES="AUTOCONF AUTOMAKE BISON DTC LIBTOOL M4 MAKE"
+# end of Companion tools
diff --git a/Engine/Build/BatchFiles/Linux/Toolchain/build_linux_toolchain.sh b/Engine/Build/BatchFiles/Linux/Toolchain/build_linux_toolchain.sh
--- a/Engine/Build/BatchFiles/Linux/Toolchain/build_linux_toolchain.sh
+++ b/Engine/Build/BatchFiles/Linux/Toolchain/build_linux_toolchain.sh
@@ -3,16 +3,16 @@
 set -x
 set -eu
 
-ToolChainVersion=v22
-LLVM_VERSION_MAJOR=16
-LLVM_VERSION=${LLVM_VERSION_MAJOR}.0.6
+ToolChainVersion=v23
+LLVM_VERSION_MAJOR=18
+LLVM_VERSION=${LLVM_VERSION_MAJOR}.1.0
 LLVM_BRANCH=release/${LLVM_VERSION_MAJOR}.x
 LLVM_TAG=llvmorg-${LLVM_VERSION}
 
 LLVM_URL=[REDACTED_PROFILE_URL]/llvm-project/releases/download/llvmorg-${LLVM_VERSION}
-ZLIB_PATH=/src/v1.2.8
+ZLIB_PATH=/src/1.3
 
-ToolChainVersionName="${ToolChainVersion}_clang-${LLVM_VERSION}-centos7"
+ToolChainVersionName="${ToolChainVersion}_clang-${LLVM_VERSION}-rockylinux8"
 
 TARGETS="x86_64-unknown-linux-gnu aarch64-unknown-linux-gnueabi"
 
@@ -29,9 +29,11 @@ echo Using $CORES cores for building
 
 echo "check_certificate=off" > "$HOME/.wgetrc"
 
-if [ ! -d "/src/build/crosstool-ng" ]; then
-	# Get crosstool-ng
-	git clone [REDACTED_PROFILE_URL]/crosstool-ng -b 1.22
+if [ ! -f "/src/build/ct-ng-build.done" ]; then
+	if [ ! -d "/src/build/crosstool-ng" ]; then
+		# Get crosstool-ng
+		git clone [REDACTED_PROFILE_URL]/crosstool-ng -b crosstool-ng-1.26.0
+	fi
 
 	# Build crosstool-ng
 	pushd crosstool-ng
@@ -42,7 +44,7 @@ if [ ! -d "/src/build/crosstool-ng" ]; then
 	for arch in $TARGETS; do
 		mkdir -p build-linux-$arch
 		pushd build-linux-$arch
-		cp /src/$arch.linux.config .config
+		cp -f /src/$arch.linux.config .config
 		../crosstool-ng/ct-ng build.$CORES
 		popd
 	done
@@ -51,19 +53,16 @@ if [ ! -d "/src/build/crosstool-ng" ]; then
 	for arch in $TARGETS; do
 		mkdir -p build-windows-$arch
 		pushd build-windows-$arch
-		cp /src/$arch.windows.config .config
+		cp -f /src/$arch.windows.config .config
 		../crosstool-ng/ct-ng build.$CORES
 		popd
 	done
+	touch /src/build/ct-ng-build.done
 fi
 
 # since we are -u in the bash script and this ENV is not set it complains when source the devtoolset-7
 export MANPATH=""
 
-# cannot use this normal method as it creates its own bash, so lets just sorce the env
-# scl enable devtoolset-7 bash
-source /opt/rh/devtoolset-7/enable
-
 # need to unset this or crosstools complains
 unset LD_LIBRARY_PATH
 
@@ -73,99 +72,110 @@ unset LD_LIBRARY_PATH
 
 echo "Cloning LLVM (tag $LLVM_TAG only)"
 # clone -b can also accept tag names
-git clone [REDACTED_PROFILE_URL]/llvm-project llvm-src -b ${LLVM_TAG} --single-branch --depth 1 -c advice.detachedHead=false
-pushd llvm-src
-git -c advice.detachedHead=false checkout tags/${LLVM_TAG} -b ${LLVM_BRANCH}
-popd
+if ! [ -d llvm-src ]; then
+	git clone [REDACTED_PROFILE_URL]/llvm-project llvm-src -b ${LLVM_TAG} --single-branch --depth 1 -c advice.detachedHead=false
+	pushd llvm-src
+	git -c advice.detachedHead=false checkout tags/${LLVM_TAG} -b ${LLVM_BRANCH}
+	popd
+fi
 
 
+# Leaving these for now in case testing reveals that one or more of them are still needed
 # this fixes an issue where AT_HWCAP2 is just not defined correctly in our sysroot. This is likely due to
 # AT_HWCAP2 being around since glibc 2.18 offically, while we are still stuck on 2.17 glibc.
-patch -d llvm-src -p 1 < /src/patches/compiler-rt/manually-define-AT_HWCAP2.diff
+#patch -d llvm-src -p 1 < /src/patches/compiler-rt/manually-define-AT_HWCAP2.diff
 
 # this fixes lack of HWCAP_CRC32 in the old glibc (similar issue)
-patch -d llvm-src -p 1 < /src/patches/compiler-rt/cpu_model_define_HWCAP_CRC32.diff
+#patch -d llvm-src -p 1 < /src/patches/compiler-rt/cpu_model_define_HWCAP_CRC32.diff
 
 # move back to defaulting to dwarf 4, as if we leave to dwarf 5 libs built with dwarf5 will force everything to dwarf5
 # even if you request dwarf 4. dwarf 5 currently causes issues with dump_syms and gdb/lldb earlier versions
-patch -d llvm-src -p 1 < /src/patches/clang/default-dwarf-4.patch
+#patch -d llvm-src -p 1 < /src/patches/clang/default-dwarf-4.patch
 
 # add a patch to disable auto-upgrade of debug info. It missed clang 16.x, so this patch shouldn't be needed for clang 17.x going forward
 # See https://reviews.llvm.org/D143229 for context
-patch -d llvm-src -p 1 < /src/patches/llvm/disable-auto-upgrade-debug-info.patch
+#patch -d llvm-src -p 1 < /src/patches/llvm/disable-auto-upgrade-debug-info.patch
 
 # LLVM has just failed to support stand-alone LLD build, cheat by moving a required header into a location it can be found easily
 # if you fulling include this you end up breaking other things. [REDACTED_PROFILE_URL]/llvm-project/issues/48572
-cp -rf llvm-src/libunwind/include/mach-o/ llvm-src/llvm/include
-
-mkdir -p build-clang
-pushd build-clang
-	# CMake Error at cmake/modules/CheckCompilerVersion.cmake:40 (message):
-	#   Host GCC version should be at least 5.1 because LLVM will soon use new C++
-	#   features which your toolchain version doesn't support.  Your version is
-	#   4.8.5.  You can temporarily opt out using
-	#   LLVM_TEMPORARILY_ALLOW_OLD_TOOLCHAIN, but very soon your toolchain won't be
-	#   supported.
-	cmake3 -G "Unix Makefiles" ../llvm-src/llvm \
-		-DLLVM_ENABLE_PROJECTS=llvm\;clang\;lld\;compiler-rt \
-		-DCMAKE_BUILD_TYPE=Release \
-		-DLLVM_ENABLE_TERMINFO=OFF \
-		-DLLVM_ENABLE_LIBXML2=OFF \
-		-DLLVM_ENABLE_ZLIB=FORCE_ON \
-		-DZLIB_LIBRARY="$ZLIB_PATH/lib/Unix/x86_64-unknown-linux-gnu/libz_fPIC.a" \
-		-DZLIB_INCLUDE_DIR="$ZLIB_PATH/include/Unix/x86_64-unknown-linux-gnu" \
-		-DLLVM_ENABLE_LIBCXX=1 \
-		-DLLVM_TEMPORARILY_ALLOW_OLD_TOOLCHAIN=ON \
-		-DCMAKE_INSTALL_PREFIX=${InstallClangDir} \
-		-DLLVM_INCLUDE_BENCHMARKS=OFF \
-		-DLLVM_TARGETS_TO_BUILD="AArch64;X86" \
-		-DCLANG_REPOSITORY_STRING="github.com/llvm/llvm-project"
-
-	make -j$CORES && make install
-popd
+#cp -rf llvm-src/libunwind/include/mach-o/ llvm-src/llvm/include
+
+if [ ! -f build-clang.done ]; then
+	mkdir -p build-clang
+	pushd build-clang
+		# CMake Error at cmake/modules/CheckCompilerVersion.cmake:40 (message):
+		#   Host GCC version should be at least 5.1 because LLVM will soon use new C++
+		#   features which your toolchain version doesn't support.  Your version is
+		#   4.8.5.  You can temporarily opt out using
+		#   LLVM_TEMPORARILY_ALLOW_OLD_TOOLCHAIN, but very soon your toolchain won't be
+		#   supported.
+		cmake3 -G "Unix Makefiles" ../llvm-src/llvm \
+			-DLLVM_ENABLE_PROJECTS=llvm\;clang\;lld\;compiler-rt \
+			-DCMAKE_BUILD_TYPE=Release \
+			-DLLVM_ENABLE_TERMINFO=OFF \
+			-DLLVM_ENABLE_LIBXML2=OFF \
+			-DLLVM_ENABLE_ZLIB=FORCE_ON \
+			-DZLIB_LIBRARY="$ZLIB_PATH/lib/Unix/x86_64-unknown-linux-gnu/Release/libz.a" \
+			-DZLIB_INCLUDE_DIR="$ZLIB_PATH/include" \
+			-DLLVM_ENABLE_LIBCXX=1 \
+			-DCMAKE_INSTALL_PREFIX=${InstallClangDir} \
+			-DLLVM_INCLUDE_BENCHMARKS=OFF \
+			-DLLVM_TARGETS_TO_BUILD="AArch64;X86" \
+			-DCLANG_REPOSITORY_STRING="github.com/llvm/llvm-project" \
+	
+		make -j$CORES && make install
+	popd
+	touch build-clang.done
+fi
 
-# Copy files
-for arch in $TARGETS; do
-	echo "Copying ${arch} toolchain..."
 
-	pushd ${OutputDirLinux}/$arch/
-		chmod -R +w .
+# Copy files
+if [ ! -f copy-linux-toolchain.done ]; then
+	for arch in $TARGETS; do
+		echo "Copying ${arch} toolchain..."
 
-		if [ -d "$arch" ]; then
-			# copy $arch/include/c++ to include/c++
-			cp -r -L $arch/include .
+		pushd ${OutputDirLinux}/$arch/
+			chmod -R +w .
 
-			# copy usr lib64 and include dirs
-			mkdir -p usr
-			cp -r -L $arch/sysroot/usr/include usr
-			cp -r -L $arch/sysroot/usr/lib64 usr
-			cp -r -L $arch/sysroot/usr/lib usr
+			if [ -d "$arch" ]; then
+				# copy $arch/include/c++ to include/c++
+				cp -r -L $arch/include .
 
-			cp -r -L $arch/lib64 .
-			cp -r -L $arch/lib .
+				# copy usr lib64 and include dirs
+				mkdir -p usr
+				cp -r -L $arch/sysroot/usr/include usr
+				cp -r -L $arch/sysroot/usr/lib64 usr
+				cp -r -L $arch/sysroot/usr/lib usr
 
-			[[ -f build.log.bz2 ]] && mv build.log.bz2 ../../build-linux-$arch.log.bz2
+				cp -r -L $arch/sysroot/lib64 .
+				cp -r -L $arch/sysroot/lib .
 
-			rm -rf $arch
-		fi
-	popd
+				[[ -f build.log.bz2 ]] && mv build.log.bz2 ../../build-linux-$arch.log.bz2
 
-	echo "Copying clang..."
-	cp -L ${InstallClangDir}/bin/clang           ${OutputDirLinux}/$arch/bin/
-	cp -L ${InstallClangDir}/bin/clang++         ${OutputDirLinux}/$arch/bin/
-	cp -L ${InstallClangDir}/bin/lld             ${OutputDirLinux}/$arch/bin/
-	cp -L ${InstallClangDir}/bin/ld.lld          ${OutputDirLinux}/$arch/bin/
-	cp -L ${InstallClangDir}/bin/llvm-ar         ${OutputDirLinux}/$arch/bin/
-	cp -L ${InstallClangDir}/bin/llvm-profdata   ${OutputDirLinux}/$arch/bin/
-	cp -L ${InstallClangDir}/bin/llvm-objcopy    ${OutputDirLinux}/$arch/bin/
-	cp -L ${InstallClangDir}/bin/llvm-symbolizer ${OutputDirLinux}/$arch/bin/
+				rm -rf $arch
+			fi
+		popd
 
-	if [ "$arch" == "x86_64-unknown-linux-gnu" ]; then
-		cp -r -L ${InstallClangDir}/lib/clang ${OutputDirLinux}/$arch/lib/
-	fi
-done
+		echo "Copying clang..."
+		cp -L ${InstallClangDir}/bin/clang           ${OutputDirLinux}/$arch/bin/
+		cp -L ${InstallClangDir}/bin/clang++         ${OutputDirLinux}/$arch/bin/
+		cp -L ${InstallClangDir}/bin/lld             ${OutputDirLinux}/$arch/bin/
+		cp -L ${InstallClangDir}/bin/ld.lld          ${OutputDirLinux}/$arch/bin/
+		cp -L ${InstallClangDir}/bin/llvm-ar         ${OutputDirLinux}/$arch/bin/
+		cp -L ${InstallClangDir}/bin/llvm-profdata   ${OutputDirLinux}/$arch/bin/
+		cp -L ${InstallClangDir}/bin/llvm-objcopy    ${OutputDirLinux}/$arch/bin/
+		cp -L ${InstallClangDir}/bin/llvm-symbolizer ${OutputDirLinux}/$arch/bin/
+		cp -L ${InstallClangDir}/bin/llvm-cov        ${OutputDirLinux}/$arch/bin/
+
+		if [ "$arch" == "x86_64-unknown-linux-gnu" ]; then
+			cp -r -L ${InstallClangDir}/lib/clang ${OutputDirLinux}/$arch/lib/
+		fi
+	done
+	touch copy-linux-toolchain.done
+fi
 
 # Build compiler-rt
+if [ ! -f build-runtime.done ]; then
 for arch in $TARGETS; do
 	if [ "$arch" == "x86_64-unknown-linux-gnu" ]; then
 		# We already built it with clang
@@ -188,12 +198,13 @@ for arch in $TARGETS; do
 			-DCMAKE_C_COMPILER_TARGET="$arch" \
 			-DCMAKE_C_COMPILER=${InstallClangDir}/bin/clang \
 			-DCMAKE_CXX_COMPILER=${InstallClangDir}/bin/clang++ \
+			-DCMAKE_LD=${InstallClangDir}/bin/lld \
 			-DCMAKE_AR=${InstallClangDir}/bin/llvm-ar \
 			-DCMAKE_NM=${InstallClangDir}/bin/llvm-nm \
 			-DCMAKE_RANLIB=${InstallClangDir}/bin/llvm-ranlib \
 			-DLLVM_ENABLE_ZLIB=FORCE_ON \
-			-DZLIB_LIBRARY="$ZLIB_PATH/lib/Unix/x86_64-unknown-linux-gnu/libz_fPIC.a" \
-			-DZLIB_INCLUDE_DIR="$ZLIB_PATH/include/Unix/x86_64-unknown-linux-gnu" \
+			-DZLIB_LIBRARY="$ZLIB_PATH/lib/Unix/x86_64-unknown-linux-gnu/Release/libz.a" \
+			-DZLIB_INCLUDE_DIR="$ZLIB_PATH/include" \
 			-DCMAKE_EXE_LINKER_FLAGS="--target=$arch -L${OutputDirLinux}/$arch/lib64 --sysroot=${OutputDirLinux}/$arch -fuse-ld=lld" \
 			-DCMAKE_C_FLAGS="--target=$arch --sysroot=${OutputDirLinux}/$arch" \
 			-DCMAKE_CXX_FLAGS="--target=$arch --sysroot=${OutputDirLinux}/$arch" \
@@ -203,7 +214,9 @@ for arch in $TARGETS; do
 			-DCMAKE_INSTALL_PREFIX=../install-rt-$arch \
 			-DSANITIZER_COMMON_LINK_FLAGS="-fuse-ld=lld" \
 			-DSCUDO_LINK_FLAGS="-fuse-ld=lld" \
-			-DLLVM_CONFIG_PATH=${InstallClangDir}/bin/llvm-config
+			-DLLVM_LIBC_INCLUDE_SCUDO=ON \
+			-DCOMPILER_RT_BUILD_XRAY=OFF \
+			-DLLVM_CONFIG_PATH=${InstallClangDir}/bin/llvm-config \
 
 		make -j$CORES && make install
 
@@ -213,6 +226,9 @@ for arch in $TARGETS; do
 	cp -r install-rt-$arch/lib/* ${OutputDirLinux}/$arch/lib/clang/${LLVM_VERSION_MAJOR}/lib/
 done
 
+touch build-runtime.done
+fi
+
 # Create version file
 echo "${ToolChainVersionName}" > ${OutputDirLinux}/ToolchainVersion.txt
 
@@ -220,32 +236,38 @@ echo "${ToolChainVersionName}" > ${OutputDirLinux}/ToolchainVersion.txt
 # Windows
 #
 
-for arch in $TARGETS; do
-	echo "Copying Windows $arch toolchain..."
+if [ ! -f copy-windows-toolchain.done ]; then
+	for arch in $TARGETS; do
+		echo "Copying Windows $arch toolchain..."
 
-	pushd ${OutputDirWindows}/$arch/
-		chmod -R +w .
+		pushd ${OutputDirWindows}/$arch/
+			chmod -R +w .
 
-		# copy $arch/include/c++ to include/c++
-		cp -r -L $arch/include .
+			# copy $arch/include/c++ to include/c++
+			cp -r -L $arch/include .
 
-		# copy usr lib64 and include dirs
-		mkdir -p usr
-		cp -r -L $arch/sysroot/usr/include usr
-		cp -r -L $arch/sysroot/usr/lib64 usr
-		cp -r -L $arch/sysroot/usr/lib usr
+			# copy usr lib64 and include dirs
+			mkdir -p usr
+			cp -r -L $arch/sysroot/usr/include usr
+			cp -r -L $arch/sysroot/usr/lib64 usr
+			cp -r -L $arch/sysroot/usr/lib usr
 
-		cp -r -L $arch/lib64 .
-		cp -r -L $arch/lib .
+			cp -r -L $arch/sysroot/lib64 .
+			cp -r -L $arch/sysroot/lib .
 
-		# Copy compiler-rt
-		cp -r -L ${OutputDirLinux}/$arch/lib/clang lib/
+			# Copy compiler-rt
+			cp -r -L ${OutputDirLinux}/$arch/lib/clang lib/
 
-		[[ -f build.log.bz2 ]] && mv build.log.bz2 ../../build-windows-$arch.log.bz2
+			# Copy linux llvm-symbolizer so that we can resolve ASan hits
+			cp -r -L ${OutputDirLinux}/$arch/bin/llvm-symbolizer bin
 
-		rm -rf $arch
-	popd
-done
+			[[ -f build.log.bz2 ]] && mv build.log.bz2 ../../build-windows-$arch.log.bz2
+
+			rm -rf $arch
+		popd
+	done
+	touch copy-windows-toolchain.done
+fi
 
 # Pack Linux files
 pushd ${OutputDirLinux}
@@ -254,29 +276,33 @@ pushd ${OutputDirLinux}
 	cp /src/build/build-linux-aarch64-unknown-linux-gnueabi/.build/tarballs/* build/src
 	tar czfvh /src/build/llvm-${LLVM_VERSION}-github-snapshot.src.tar.gz --hard-dereference /src/build/llvm-src
 	cp /src/build/*.src.tar.gz build/src
-	cp /src/*.{config,sh,nsi,bat} build/scripts
+	cp -f /src/*.{config,sh,nsi,bat} build/scripts
 
 	# copy the toolchain in the directory named its version as per convention
-	mkdir ${OutputDirLinux}/${ToolChainVersionName}
-	cp -r x86_64-unknown-linux-gnu ${OutputDirLinux}/${ToolChainVersionName}
-	cp -r build ${OutputDirLinux}/${ToolChainVersionName}
-	cp -r aarch64-unknown-linux-gnueabi ${OutputDirLinux}/${ToolChainVersionName}
-	cp -r ToolchainVersion.txt ${OutputDirLinux}/${ToolChainVersionName}
+	mkdir -p ${OutputDirLinux}/${ToolChainVersionName}
+	cp -rf x86_64-unknown-linux-gnu ${OutputDirLinux}/${ToolChainVersionName}
+	cp -rf build ${OutputDirLinux}/${ToolChainVersionName}
+	cp -rf aarch64-unknown-linux-gnueabi ${OutputDirLinux}/${ToolChainVersionName}
+	cp -rf ToolchainVersion.txt ${OutputDirLinux}/${ToolChainVersionName}
 
 	# delete libraries in x86_64's lib folder or bundled binares with crash
 	find ${OutputDirLinux}/${ToolChainVersionName}/x86_64-unknown-linux-gnu/lib/ -maxdepth 1 -type f -delete
 
-	tar czfhv native-linux-${ToolChainVersionName}.tar.gz --hard-dereference ${ToolChainVersionName}
+	# remove broken links before tar otherwise it will fail
+	# currently, {sysroot}/lib/bfd-plugins/liblto_plugin points to a non-existent liblto_plugin in {sysroot}/libexec
+	find ${OutputDirLinux} -xtype l -delete
+
+	tar czfhv /src/build/native-linux-${ToolChainVersionName}.tar.gz --hard-dereference ${ToolChainVersionName}
 popd
 
 # Pack Windows files
 pushd ${OutputDirWindows}
 	mkdir -p build/{src,scripts}
-	cp /src/build/build-windows-x86_64-unknown-linux-gnu/.build/tarballs/* build/src
-	cp /src/build/build-windows-aarch64-unknown-linux-gnueabi/.build/tarballs/* build/src
+	cp -f /src/build/build-windows-x86_64-unknown-linux-gnu/.build/tarballs/* build/src
+	cp -f /src/build/build-windows-aarch64-unknown-linux-gnueabi/.build/tarballs/* build/src
 	zip -r /src/build/llvm-${LLVM_VERSION}-github-snapshot.src.zip /src/build/llvm-src
-	cp /src/build/*.src.zip build/src
-	cp /src/*.{config,sh,nsi,bat} build/scripts
+	cp -f /src/build/*.src.zip build/src
+	cp -f /src/*.{config,sh,nsi,bat} build/scripts
 
 	zip -r /src/build/${ToolChainVersionName}-windows.zip *
 popd
diff --git a/Engine/Build/BatchFiles/Linux/Toolchain/x86_64-unknown-linux-gnu.linux.config b/Engine/Build/BatchFiles/Linux/Toolchain/x86_64-unknown-linux-gnu.linux.config
--- a/Engine/Build/BatchFiles/Linux/Toolchain/x86_64-unknown-linux-gnu.linux.config
+++ b/Engine/Build/BatchFiles/Linux/Toolchain/x86_64-unknown-linux-gnu.linux.config
@@ -1,10 +1,34 @@
 #
 # Automatically generated file; DO NOT EDIT.
-# Crosstool-NG Configuration
-#
-CT_CONFIGURE_has_make381=y
-CT_CONFIGURE_has_xz=y
-CT_CONFIGURE_has_cvs=y
+# crosstool-NG 1.26.0 Configuration
+#
+CT_CONFIGURE_has_static_link=y
+CT_CONFIGURE_has_cxx11=y
+CT_CONFIGURE_has_wget=y
+CT_CONFIGURE_has_curl=y
+CT_CONFIGURE_has_rsync=y
+CT_CONFIGURE_has_make_3_81_or_newer=y
+CT_CONFIGURE_has_make_4_0_or_newer=y
+CT_CONFIGURE_has_libtool_2_4_or_newer=y
+CT_CONFIGURE_has_libtoolize_2_4_or_newer=y
+CT_CONFIGURE_has_autoconf_2_65_or_newer=y
+CT_CONFIGURE_has_autoreconf_2_65_or_newer=y
+CT_CONFIGURE_has_automake_1_15_or_newer=y
+CT_CONFIGURE_has_gnu_m4_1_4_12_or_newer=y
+CT_CONFIGURE_has_bison_2_7_or_newer=y
+CT_CONFIGURE_has_bison_3_0_4_or_newer=y
+CT_CONFIGURE_has_python=y
+CT_CONFIGURE_has_git=y
+CT_CONFIGURE_has_md5sum=y
+CT_CONFIGURE_has_sha1sum=y
+CT_CONFIGURE_has_sha256sum=y
+CT_CONFIGURE_has_sha512sum=y
+CT_CONFIGURE_has_install_with_strip_program=y
+CT_VERSION="1.26.0"
+CT_VCHECK=""
+CT_CONFIG_VERSION_ENV="4"
+CT_CONFIG_VERSION_CURRENT="4"
+CT_CONFIG_VERSION="4"
 CT_MODULES=y
 
 #
@@ -21,41 +45,49 @@ CT_MODULES=y
 #
 # Paths
 #
-CT_LOCAL_TARBALLS_DIR=""
+CT_LOCAL_TARBALLS_DIR="${HOME}/src"
+CT_SAVE_TARBALLS=y
+# CT_TARBALLS_BUILDROOT_LAYOUT is not set
 CT_WORK_DIR="${CT_TOP_DIR}/.build"
+CT_BUILD_TOP_DIR="${CT_WORK_DIR:-${CT_TOP_DIR}/.build}/${CT_HOST:+HOST-${CT_HOST}/}${CT_TARGET}"
+CT_BUILD_DIR="${CT_BUILD_TOP_DIR}/build"
 CT_PREFIX_DIR="${CT_TOP_DIR}/../OUTPUT-linux/${CT_TARGET}"
-CT_INSTALL_DIR="${CT_PREFIX_DIR}"
 CT_RM_RF_PREFIX_DIR=y
 CT_REMOVE_DOCS=y
-CT_INSTALL_DIR_RO=y
+CT_INSTALL_LICENSES=y
+CT_PREFIX_DIR_RO=y
 CT_STRIP_HOST_TOOLCHAIN_EXECUTABLES=y
 # CT_STRIP_TARGET_TOOLCHAIN_EXECUTABLES is not set
 
 #
 # Downloading
 #
+CT_DOWNLOAD_AGENT_WGET=y
+# CT_DOWNLOAD_AGENT_CURL is not set
+# CT_DOWNLOAD_AGENT_NONE is not set
 # CT_FORBID_DOWNLOAD is not set
 # CT_FORCE_DOWNLOAD is not set
 CT_CONNECT_TIMEOUT=10
+CT_DOWNLOAD_WGET_OPTIONS="--passive-ftp --tries=3 -nc --progress=dot:binary"
 # CT_ONLY_DOWNLOAD is not set
 # CT_USE_MIRROR is not set
+CT_VERIFY_DOWNLOAD_DIGEST=y
+CT_VERIFY_DOWNLOAD_DIGEST_SHA512=y
+# CT_VERIFY_DOWNLOAD_DIGEST_SHA256 is not set
+# CT_VERIFY_DOWNLOAD_DIGEST_SHA1 is not set
+# CT_VERIFY_DOWNLOAD_DIGEST_MD5 is not set
+CT_VERIFY_DOWNLOAD_DIGEST_ALG="sha512"
+# CT_VERIFY_DOWNLOAD_SIGNATURE is not set
 
 #
 # Extracting
 #
 # CT_FORCE_EXTRACT is not set
-CT_OVERIDE_CONFIG_GUESS_SUB=y
+CT_OVERRIDE_CONFIG_GUESS_SUB=y
 # CT_ONLY_EXTRACT is not set
-# CT_PATCH_BUNDLED is not set
-# CT_PATCH_LOCAL is not set
-CT_PATCH_BUNDLED_LOCAL=y
-# CT_PATCH_LOCAL_BUNDLED is not set
-# CT_PATCH_BUNDLED_FALLBACK_LOCAL is not set
-# CT_PATCH_LOCAL_FALLBACK_BUNDLED is not set
-# CT_PATCH_NONE is not set
-CT_PATCH_ORDER="bundled,local"
-CT_PATCH_USE_LOCAL=y
-CT_LOCAL_PATCH_DIR="${CT_TOP_DIR}/../../patches"
+CT_PATCH_BUNDLED=y
+# CT_PATCH_BUNDLED_LOCAL is not set
+CT_PATCH_ORDER="bundled"
 
 #
 # Build behavior
@@ -64,6 +96,7 @@ CT_PARALLEL_JOBS=0
 CT_LOAD=""
 CT_USE_PIPES=y
 CT_EXTRA_CFLAGS_FOR_BUILD=""
+CT_EXTRA_CXXFLAGS_FOR_BUILD=""
 CT_EXTRA_LDFLAGS_FOR_BUILD=""
 CT_EXTRA_CFLAGS_FOR_HOST=""
 CT_EXTRA_LDFLAGS_FOR_HOST=""
@@ -78,71 +111,74 @@ CT_CONFIG_SHELL="${bash}"
 #
 # CT_LOG_ERROR is not set
 # CT_LOG_WARN is not set
-CT_LOG_INFO=y
-# CT_LOG_EXTRA is not set
+# CT_LOG_INFO is not set
+CT_LOG_EXTRA=y
 # CT_LOG_ALL is not set
 # CT_LOG_DEBUG is not set
-CT_LOG_LEVEL_MAX="INFO"
+CT_LOG_LEVEL_MAX="EXTRA"
 # CT_LOG_SEE_TOOLS_WARN is not set
 CT_LOG_PROGRESS_BAR=y
 CT_LOG_TO_FILE=y
 CT_LOG_FILE_COMPRESS=y
+# end of Paths and misc options
 
 #
 # Target options
 #
+# CT_ARCH_ALPHA is not set
+# CT_ARCH_ARC is not set
+# CT_ARCH_ARM is not set
+# CT_ARCH_AVR is not set
+# CT_ARCH_BPF is not set
+# CT_ARCH_M68K is not set
+# CT_ARCH_MIPS is not set
+# CT_ARCH_NIOS2 is not set
+# CT_ARCH_POWERPC is not set
+# CT_ARCH_PRU is not set
+# CT_ARCH_S390 is not set
+# CT_ARCH_SH is not set
+# CT_ARCH_SPARC is not set
+CT_ARCH_X86=y
+# CT_ARCH_XTENSA is not set
 CT_ARCH="x86"
-CT_ARCH_SUPPORTS_32=y
-CT_ARCH_SUPPORTS_64=y
-CT_ARCH_SUPPORTS_WITH_ARCH=y
-CT_ARCH_SUPPORTS_WITH_CPU=y
-CT_ARCH_SUPPORTS_WITH_TUNE=y
-CT_ARCH_DEFAULT_32=y
-CT_ARCH_ARCH=""
+CT_ARCH_CHOICE_KSYM="X86"
 CT_ARCH_CPU=""
 CT_ARCH_TUNE=""
-# CT_ARCH_32 is not set
-CT_ARCH_64=y
-CT_ARCH_BITNESS=64
-CT_TARGET_CFLAGS=""
-CT_TARGET_LDFLAGS=""
-# CT_ARCH_alpha is not set
-# CT_ARCH_arm is not set
-# CT_ARCH_avr is not set
-# CT_ARCH_m68k is not set
-# CT_ARCH_mips is not set
-# CT_ARCH_nios2 is not set
-# CT_ARCH_powerpc is not set
-# CT_ARCH_s390 is not set
-# CT_ARCH_sh is not set
-# CT_ARCH_sparc is not set
-CT_ARCH_x86=y
-# CT_ARCH_xtensa is not set
-CT_ARCH_alpha_AVAILABLE=y
-CT_ARCH_arm_AVAILABLE=y
-CT_ARCH_avr_AVAILABLE=y
-CT_ARCH_m68k_AVAILABLE=y
-CT_ARCH_microblaze_AVAILABLE=y
-CT_ARCH_mips_AVAILABLE=y
-CT_ARCH_nios2_AVAILABLE=y
-CT_ARCH_powerpc_AVAILABLE=y
-CT_ARCH_s390_AVAILABLE=y
-CT_ARCH_sh_AVAILABLE=y
-CT_ARCH_sparc_AVAILABLE=y
-CT_ARCH_x86_AVAILABLE=y
-CT_ARCH_xtensa_AVAILABLE=y
+CT_ARCH_X86_SHOW=y
+
+#
+# Options for x86
+#
+CT_ARCH_X86_PKG_KSYM=""
+CT_ALL_ARCH_CHOICES="ALPHA ARC ARM AVR BPF C6X LOONGARCH M68K MICROBLAZE MIPS MOXIE MSP430 NIOS2 POWERPC PRU RISCV S390 SH SPARC X86 XTENSA"
 CT_ARCH_SUFFIX=""
+# CT_OMIT_TARGET_VENDOR is not set
 
 #
 # Generic target options
 #
 # CT_MULTILIB is not set
+CT_DEMULTILIB=y
 CT_ARCH_USE_MMU=y
+CT_ARCH_SUPPORTS_LIBSANITIZER=y
+CT_ARCH_SUPPORTS_32=y
+CT_ARCH_SUPPORTS_64=y
+CT_ARCH_DEFAULT_32=y
+CT_ARCH_BITNESS=64
+# CT_ARCH_32 is not set
+CT_ARCH_64=y
+CT_ARCH_SUPPORTS_WITH_32_64=y
 
 #
 # Target optimisations
 #
-CT_ARCH_FLOAT=""
+CT_ARCH_SUPPORTS_WITH_ARCH=y
+CT_ARCH_SUPPORTS_WITH_CPU=y
+CT_ARCH_SUPPORTS_WITH_TUNE=y
+CT_ARCH_ARCH=""
+CT_TARGET_CFLAGS=""
+CT_TARGET_LDFLAGS=""
+# end of Target options
 
 #
 # Toolchain options
@@ -156,7 +192,9 @@ CT_USE_SYSROOT=y
 CT_SYSROOT_NAME="sysroot"
 CT_SYSROOT_DIR_PREFIX=""
 CT_WANTS_STATIC_LINK=y
+CT_WANTS_STATIC_LINK_CXX=y
 CT_STATIC_TOOLCHAIN=y
+CT_SHOW_CT_VERSION=y
 CT_TOOLCHAIN_PKGVERSION=""
 CT_TOOLCHAIN_BUGURL=""
 
@@ -185,132 +223,279 @@ CT_BUILD_SUFFIX=""
 # Misc options
 #
 # CT_TOOLCHAIN_ENABLE_NLS is not set
+# end of Toolchain options
 
 #
 # Operating System
 #
 CT_KERNEL_SUPPORTS_SHARED_LIBS=y
+# CT_KERNEL_BARE_METAL is not set
+CT_KERNEL_LINUX=y
 CT_KERNEL="linux"
-CT_KERNEL_VERSION="3.10.93"
-# CT_KERNEL_bare_metal is not set
-CT_KERNEL_linux=y
-CT_KERNEL_bare_metal_AVAILABLE=y
-CT_KERNEL_linux_AVAILABLE=y
-# CT_KERNEL_V_4_3 is not set
-# CT_KERNEL_V_4_2 is not set
-# CT_KERNEL_V_4_1 is not set
-# CT_KERNEL_V_3_18 is not set
-# CT_KERNEL_V_3_14 is not set
-# CT_KERNEL_V_3_12 is not set
-CT_KERNEL_V_3_10=y
-# CT_KERNEL_V_3_4 is not set
-# CT_KERNEL_V_3_2 is not set
-# CT_KERNEL_V_2_6_32 is not set
-# CT_KERNEL_LINUX_CUSTOM is not set
-CT_KERNEL_windows_AVAILABLE=y
-
-#
-# Common kernel options
-#
-CT_SHARED_LIBS=y
-
-#
-# linux other options
-#
+CT_KERNEL_CHOICE_KSYM="LINUX"
+CT_KERNEL_LINUX_SHOW=y
+
+#
+# Options for linux
+#
+CT_KERNEL_LINUX_PKG_KSYM="LINUX"
+CT_LINUX_DIR_NAME="linux"
+CT_LINUX_USE_WWW_KERNEL_ORG=y
+# CT_LINUX_USE_ORACLE is not set
+CT_LINUX_USE="LINUX"
+CT_LINUX_PKG_NAME="linux"
+CT_LINUX_SRC_RELEASE=y
+# CT_LINUX_SRC_DEVEL is not set
+CT_LINUX_PATCH_ORDER="global"
+# CT_LINUX_V_6_4 is not set
+# CT_LINUX_V_6_3 is not set
+# CT_LINUX_V_6_2 is not set
+# CT_LINUX_V_6_1 is not set
+# CT_LINUX_V_6_0 is not set
+# CT_LINUX_V_5_19 is not set
+# CT_LINUX_V_5_18 is not set
+# CT_LINUX_V_5_17 is not set
+# CT_LINUX_V_5_16 is not set
+# CT_LINUX_V_5_15 is not set
+# CT_LINUX_V_5_14 is not set
+# CT_LINUX_V_5_13 is not set
+# CT_LINUX_V_5_12 is not set
+# CT_LINUX_V_5_11 is not set
+# CT_LINUX_V_5_10 is not set
+# CT_LINUX_V_5_9 is not set
+# CT_LINUX_V_5_8 is not set
+# CT_LINUX_V_5_7 is not set
+# CT_LINUX_V_5_4 is not set
+# CT_LINUX_V_5_3 is not set
+# CT_LINUX_V_5_2 is not set
+# CT_LINUX_V_5_1 is not set
+# CT_LINUX_V_5_0 is not set
+# CT_LINUX_V_4_20 is not set
+# CT_LINUX_V_4_19 is not set
+CT_LINUX_V_4_18=y
+# CT_LINUX_V_4_17 is not set
+# CT_LINUX_V_4_16 is not set
+# CT_LINUX_V_4_15 is not set
+# CT_LINUX_V_4_14 is not set
+# CT_LINUX_V_4_13 is not set
+# CT_LINUX_V_4_12 is not set
+# CT_LINUX_V_4_11 is not set
+# CT_LINUX_V_4_10 is not set
+# CT_LINUX_V_4_9 is not set
+# CT_LINUX_V_4_4 is not set
+# CT_LINUX_V_4_1 is not set
+# CT_LINUX_V_3_16 is not set
+# CT_LINUX_V_3_13 is not set
+# CT_LINUX_V_3_12 is not set
+# CT_LINUX_V_3_10 is not set
+# CT_LINUX_V_3_4 is not set
+# CT_LINUX_V_3_2 is not set
+CT_LINUX_VERSION="4.18.20"
+CT_LINUX_MIRRORS="$(CT_Mirrors kernel.org linux ${CT_LINUX_VERSION})"
+CT_LINUX_ARCHIVE_FILENAME="@{pkg_name}-@{version}"
+CT_LINUX_ARCHIVE_DIRNAME="@{pkg_name}-@{version}"
+CT_LINUX_ARCHIVE_FORMATS=".tar.xz .tar.gz"
+CT_LINUX_SIGNATURE_FORMAT="unpacked/.sign"
+CT_LINUX_5_19_or_older=y
+CT_LINUX_older_than_5_19=y
+CT_LINUX_5_12_or_older=y
+CT_LINUX_older_than_5_12=y
+CT_LINUX_5_5_or_older=y
+CT_LINUX_older_than_5_5=y
+CT_LINUX_5_3_or_older=y
+CT_LINUX_older_than_5_3=y
+CT_LINUX_later_than_4_8=y
+CT_LINUX_4_8_or_later=y
+CT_LINUX_later_than_3_7=y
+CT_LINUX_3_7_or_later=y
+CT_LINUX_later_than_3_2=y
+CT_LINUX_3_2_or_later=y
+CT_LINUX_REQUIRE_3_2_or_later=y
+CT_KERNEL_DEP_RSYNC=y
 CT_KERNEL_LINUX_VERBOSITY_0=y
 # CT_KERNEL_LINUX_VERBOSITY_1 is not set
 # CT_KERNEL_LINUX_VERBOSITY_2 is not set
 CT_KERNEL_LINUX_VERBOSE_LEVEL=0
 CT_KERNEL_LINUX_INSTALL_CHECK=y
+CT_ALL_KERNEL_CHOICES="BARE_METAL LINUX WINDOWS"
+
+#
+# Common kernel options
+#
+CT_SHARED_LIBS=y
+# end of Operating System
 
 #
 # Binary utilities
 #
 CT_ARCH_BINFMT_ELF=y
+CT_BINUTILS_BINUTILS=y
 CT_BINUTILS="binutils"
-CT_BINUTILS_binutils=y
+CT_BINUTILS_CHOICE_KSYM="BINUTILS"
+CT_BINUTILS_BINUTILS_SHOW=y
+
+#
+# Options for binutils
+#
+CT_BINUTILS_BINUTILS_PKG_KSYM="BINUTILS"
+CT_BINUTILS_DIR_NAME="binutils"
+CT_BINUTILS_USE_GNU=y
+# CT_BINUTILS_USE_ORACLE is not set
+CT_BINUTILS_USE="BINUTILS"
+CT_BINUTILS_PKG_NAME="binutils"
+CT_BINUTILS_SRC_RELEASE=y
+# CT_BINUTILS_SRC_DEVEL is not set
+CT_BINUTILS_PATCH_ORDER="global"
+CT_BINUTILS_V_2_40=y
+# CT_BINUTILS_V_2_39 is not set
+# CT_BINUTILS_V_2_38 is not set
+# CT_BINUTILS_V_2_37 is not set
+# CT_BINUTILS_V_2_36 is not set
+# CT_BINUTILS_V_2_35 is not set
+# CT_BINUTILS_V_2_34 is not set
+# CT_BINUTILS_V_2_33 is not set
+# CT_BINUTILS_V_2_32 is not set
+# CT_BINUTILS_V_2_31 is not set
+# CT_BINUTILS_V_2_30 is not set
+# CT_BINUTILS_V_2_29 is not set
+# CT_BINUTILS_V_2_28 is not set
+# CT_BINUTILS_V_2_27 is not set
+# CT_BINUTILS_V_2_26 is not set
+CT_BINUTILS_VERSION="2.40"
+CT_BINUTILS_MIRRORS="$(CT_Mirrors GNU binutils) $(CT_Mirrors sourceware binutils/releases)"
+CT_BINUTILS_ARCHIVE_FILENAME="@{pkg_name}-@{version}"
+CT_BINUTILS_ARCHIVE_DIRNAME="@{pkg_name}-@{version}"
+CT_BINUTILS_ARCHIVE_FORMATS=".tar.xz .tar.bz2 .tar.gz"
+CT_BINUTILS_SIGNATURE_FORMAT="packed/.sig"
+CT_BINUTILS_later_than_2_39=y
+CT_BINUTILS_2_39_or_later=y
+CT_BINUTILS_later_than_2_30=y
+CT_BINUTILS_2_30_or_later=y
+CT_BINUTILS_later_than_2_27=y
+CT_BINUTILS_2_27_or_later=y
+CT_BINUTILS_later_than_2_26=y
+CT_BINUTILS_2_26_or_later=y
 
 #
 # GNU binutils
 #
-# CT_CC_BINUTILS_SHOW_LINARO is not set
-CT_BINUTILS_V_2_31_1=y
-# CT_BINUTILS_V_2_25 is not set
-# CT_BINUTILS_V_2_24 is not set
-# CT_BINUTILS_V_2_23_2 is not set
-# CT_BINUTILS_V_2_23_1 is not set
-# CT_BINUTILS_V_2_22 is not set
-# CT_BINUTILS_V_2_21_53 is not set
-# CT_BINUTILS_V_2_21_1a is not set
-# CT_BINUTILS_V_2_20_1a is not set
-# CT_BINUTILS_V_2_19_1a is not set
-# CT_BINUTILS_V_2_18a is not set
-CT_BINUTILS_VERSION="2.31.1"
-CT_BINUTILS_2_25_1_or_later=y
-CT_BINUTILS_2_25_or_later=y
-CT_BINUTILS_2_24_or_later=y
-CT_BINUTILS_2_23_or_later=y
-CT_BINUTILS_2_22_or_later=y
-CT_BINUTILS_2_21_or_later=y
-CT_BINUTILS_2_20_or_later=y
-CT_BINUTILS_2_19_or_later=y
-CT_BINUTILS_2_18_or_later=y
-CT_BINUTILS_HAS_HASH_STYLE=y
-CT_BINUTILS_HAS_GOLD=y
 CT_BINUTILS_GOLD_SUPPORTS_ARCH=y
-CT_BINUTILS_GOLD_SUPPORT=y
-CT_BINUTILS_HAS_PLUGINS=y
-CT_BINUTILS_HAS_PKGVERSION_BUGURL=y
-CT_BINUTILS_FORCE_LD_BFD=y
+CT_BINUTILS_FORCE_LD_BFD_DEFAULT=y
 CT_BINUTILS_LINKER_LD=y
-# CT_BINUTILS_LINKER_LD_GOLD is not set
-# CT_BINUTILS_LINKER_GOLD_LD is not set
 CT_BINUTILS_LINKERS_LIST="ld"
 CT_BINUTILS_LINKER_DEFAULT="bfd"
-# CT_BINUTILS_PLUGINS is not set
+CT_BINUTILS_RELRO=m
+CT_BINUTILS_DETERMINISTIC_ARCHIVES=y
 CT_BINUTILS_EXTRA_CONFIG_ARRAY=""
 # CT_BINUTILS_FOR_TARGET is not set
-
-#
-# binutils other options
-#
+CT_ALL_BINUTILS_CHOICES="BINUTILS"
+# end of Binary utilities
 
 #
 # C-library
 #
+CT_LIBC_GLIBC=y
+# CT_LIBC_UCLIBC_NG is not set
 CT_LIBC="glibc"
-CT_LIBC_VERSION="2.17"
-CT_LIBC_glibc=y
-# CT_LIBC_musl is not set
-# CT_LIBC_uClibc is not set
-CT_LIBC_avr_libc_AVAILABLE=y
-CT_LIBC_glibc_AVAILABLE=y
+CT_LIBC_CHOICE_KSYM="GLIBC"
+CT_LIBC_GLIBC_SHOW=y
+
+#
+# Options for glibc
+#
+CT_LIBC_GLIBC_PKG_KSYM="GLIBC"
+CT_GLIBC_DIR_NAME="glibc"
+CT_GLIBC_USE_GNU=y
+# CT_GLIBC_USE_ORACLE is not set
+CT_GLIBC_USE="GLIBC"
+CT_GLIBC_PKG_NAME="glibc"
+CT_GLIBC_SRC_RELEASE=y
+# CT_GLIBC_SRC_DEVEL is not set
+CT_GLIBC_PATCH_ORDER="global"
+CT_GLIBC_V_2_28=y
+# CT_GLIBC_V_2_27 is not set
+# CT_GLIBC_V_2_26 is not set
+# CT_GLIBC_V_2_25 is not set
+# CT_GLIBC_V_2_24 is not set
+# CT_GLIBC_V_2_23 is not set
+# CT_GLIBC_V_2_19 is not set
+# CT_GLIBC_V_2_17 is not set
+CT_GLIBC_VERSION="2.28"
+CT_GLIBC_MIRRORS="$(CT_Mirrors GNU glibc)"
+CT_GLIBC_ARCHIVE_FILENAME="@{pkg_name}-@{version}"
+CT_GLIBC_ARCHIVE_DIRNAME="@{pkg_name}-@{version}"
+CT_GLIBC_ARCHIVE_FORMATS=".tar.xz .tar.bz2 .tar.gz"
+CT_GLIBC_SIGNATURE_FORMAT="packed/.sig"
+CT_GLIBC_2_38_or_older=y
+CT_GLIBC_older_than_2_38=y
+CT_GLIBC_2_37_or_older=y
+CT_GLIBC_older_than_2_37=y
+CT_GLIBC_2_36_or_older=y
+CT_GLIBC_older_than_2_36=y
+CT_GLIBC_2_34_or_older=y
+CT_GLIBC_older_than_2_34=y
+CT_GLIBC_2_32_or_older=y
+CT_GLIBC_older_than_2_32=y
+CT_GLIBC_2_31_or_older=y
+CT_GLIBC_older_than_2_31=y
+CT_GLIBC_2_30_or_older=y
+CT_GLIBC_older_than_2_30=y
+CT_GLIBC_2_29_or_older=y
+CT_GLIBC_older_than_2_29=y
+CT_GLIBC_REQUIRE_older_than_2_29=y
+CT_GLIBC_2_28_or_later=y
+CT_GLIBC_2_28_or_older=y
+CT_GLIBC_later_than_2_27=y
+CT_GLIBC_2_27_or_later=y
+CT_GLIBC_later_than_2_26=y
+CT_GLIBC_2_26_or_later=y
+CT_GLIBC_later_than_2_25=y
+CT_GLIBC_2_25_or_later=y
+CT_GLIBC_later_than_2_24=y
+CT_GLIBC_2_24_or_later=y
+CT_GLIBC_later_than_2_23=y
+CT_GLIBC_2_23_or_later=y
+CT_GLIBC_later_than_2_20=y
+CT_GLIBC_2_20_or_later=y
+CT_GLIBC_later_than_2_17=y
+CT_GLIBC_2_17_or_later=y
+CT_GLIBC_later_than_2_14=y
+CT_GLIBC_2_14_or_later=y
+CT_GLIBC_DEP_KERNEL_HEADERS_VERSION=y
+CT_GLIBC_DEP_BINUTILS=y
+CT_GLIBC_DEP_GCC=y
+CT_GLIBC_DEP_PYTHON=y
+
+#
+# glibc >=2.29 requires python >3.4
+#
+CT_GLIBC_SPARC_ALLOW_V7=y
 CT_THREADS="nptl"
-# CT_CC_GLIBC_SHOW_LINARO is not set
-# CT_LIBC_GLIBC_V_2_22 is not set
-# CT_LIBC_GLIBC_V_2_21 is not set
-# CT_LIBC_GLIBC_V_2_20 is not set
-# CT_LIBC_GLIBC_V_2_19 is not set
-# CT_LIBC_GLIBC_V_2_18 is not set
-CT_LIBC_GLIBC_V_2_17=y
-# CT_LIBC_GLIBC_V_2_16_0 is not set
-# CT_LIBC_GLIBC_V_2_15 is not set
-# CT_LIBC_GLIBC_V_2_14_1 is not set
-# CT_LIBC_GLIBC_V_2_14 is not set
-# CT_LIBC_GLIBC_V_2_13 is not set
-# CT_LIBC_GLIBC_V_2_12_2 is not set
-# CT_LIBC_GLIBC_V_2_12_1 is not set
-# CT_LIBC_GLIBC_V_2_11_1 is not set
-# CT_LIBC_GLIBC_V_2_11 is not set
-# CT_LIBC_GLIBC_V_2_10_1 is not set
-# CT_LIBC_GLIBC_V_2_9 is not set
-# CT_LIBC_GLIBC_V_2_8 is not set
-CT_LIBC_GLIBC_2_17_or_later=y
-CT_LIBC_mingw_AVAILABLE=y
-CT_LIBC_musl_AVAILABLE=y
-CT_LIBC_newlib_AVAILABLE=y
-CT_LIBC_none_AVAILABLE=y
-CT_LIBC_uClibc_AVAILABLE=y
+CT_GLIBC_BUILD_SSP=y
+CT_GLIBC_HAS_LIBIDN_ADDON=y
+# CT_GLIBC_USE_LIBIDN_ADDON is not set
+CT_GLIBC_NO_SPARC_V8=y
+CT_GLIBC_HAS_OBSOLETE_RPC=y
+CT_GLIBC_EXTRA_CONFIG_ARRAY=""
+CT_GLIBC_CONFIGPARMS=""
+CT_GLIBC_ENABLE_DEBUG=y
+CT_GLIBC_EXTRA_CFLAGS=""
+CT_GLIBC_ENABLE_OBSOLETE_RPC=y
+# CT_GLIBC_DISABLE_VERSIONING is not set
+CT_GLIBC_OLDEST_ABI=""
+CT_GLIBC_FORCE_UNWIND=y
+# CT_GLIBC_LOCALES is not set
+CT_GLIBC_KERNEL_VERSION_NONE=y
+# CT_GLIBC_KERNEL_VERSION_AS_HEADERS is not set
+# CT_GLIBC_KERNEL_VERSION_CHOSEN is not set
+CT_GLIBC_MIN_KERNEL=""
+CT_GLIBC_SSP_DEFAULT=y
+# CT_GLIBC_SSP_NO is not set
+# CT_GLIBC_SSP_YES is not set
+# CT_GLIBC_SSP_ALL is not set
+# CT_GLIBC_SSP_STRONG is not set
+# CT_GLIBC_ENABLE_COMMON_FLAG is not set
+CT_ALL_LIBC_CHOICES="AVR_LIBC GLIBC MINGW_W64 MOXIEBOX MUSL NEWLIB NONE PICOLIBC UCLIBC_NG"
 CT_LIBC_SUPPORT_THREADS_ANY=y
 CT_LIBC_SUPPORT_THREADS_NATIVE=y
 
@@ -318,95 +503,107 @@ CT_LIBC_SUPPORT_THREADS_NATIVE=y
 # Common C library options
 #
 CT_THREADS_NATIVE=y
+# CT_CREATE_LDSO_CONF is not set
 CT_LIBC_XLDD=y
-
-#
-# glibc other options
-#
-# CT_LIBC_GLIBC_PORTS_EXTERNAL is not set
-CT_LIBC_glibc_familly=y
-CT_LIBC_GLIBC_EXTRA_CONFIG_ARRAY=""
-CT_LIBC_GLIBC_CONFIGPARMS=""
-CT_LIBC_GLIBC_EXTRA_CFLAGS=""
-CT_LIBC_EXTRA_CC_ARGS=""
-# CT_LIBC_DISABLE_VERSIONING is not set
-CT_LIBC_OLDEST_ABI=""
-CT_LIBC_GLIBC_FORCE_UNWIND=y
-# CT_LIBC_GLIBC_USE_PORTS is not set
-CT_LIBC_ADDONS_LIST=""
-# CT_LIBC_LOCALES is not set
-# CT_LIBC_GLIBC_KERNEL_VERSION_NONE is not set
-CT_LIBC_GLIBC_KERNEL_VERSION_AS_HEADERS=y
-# CT_LIBC_GLIBC_KERNEL_VERSION_CHOSEN is not set
-CT_LIBC_GLIBC_MIN_KERNEL="3.10.93"
+# end of C-library
 
 #
 # C compiler
 #
+CT_CC_CORE_NEEDED=y
+CT_CC_SUPPORT_CXX=y
+CT_CC_SUPPORT_FORTRAN=y
+CT_CC_SUPPORT_ADA=y
+CT_CC_SUPPORT_D=y
+CT_CC_SUPPORT_JIT=y
+CT_CC_SUPPORT_OBJC=y
+CT_CC_SUPPORT_OBJCXX=y
+CT_CC_SUPPORT_GOLANG=y
+CT_CC_GCC=y
 CT_CC="gcc"
-CT_CC_CORE_PASSES_NEEDED=y
-CT_CC_CORE_PASS_1_NEEDED=y
-CT_CC_CORE_PASS_2_NEEDED=y
-CT_CC_gcc=y
-# CT_CC_GCC_SHOW_LINARO is not set
-# CT_CC_GCC_V_5_2_0=y
-# CT_CC_GCC_V_4_9_3 is not set
-CT_CC_GCC_V_4_8_5=y
-# CT_CC_GCC_V_4_7_4 is not set
-# CT_CC_GCC_V_4_6_4 is not set
-# CT_CC_GCC_V_4_5_4 is not set
-# CT_CC_GCC_V_4_4_7 is not set
-# CT_CC_GCC_V_4_3_6 is not set
-# CT_CC_GCC_V_4_2_4 is not set
-CT_CC_GCC_4_2_or_later=y
-CT_CC_GCC_4_3_or_later=y
-CT_CC_GCC_4_4_or_later=y
-CT_CC_GCC_4_5_or_later=y
-CT_CC_GCC_4_6_or_later=y
-CT_CC_GCC_4_7_or_later=y
-CT_CC_GCC_4_8_or_later=y
-CT_CC_GCC_4_9_or_later=y
-CT_CC_GCC_5=y
-CT_CC_GCC_5_or_later=y
-CT_CC_GCC_HAS_GRAPHITE=y
-CT_CC_GCC_USE_GRAPHITE=y
-CT_CC_GCC_HAS_LTO=y
-CT_CC_GCC_USE_LTO=y
-CT_CC_GCC_HAS_PKGVERSION_BUGURL=y
-CT_CC_GCC_HAS_BUILD_ID=y
-CT_CC_GCC_HAS_LNK_HASH_STYLE=y
-CT_CC_GCC_USE_GMP_MPFR=y
-CT_CC_GCC_USE_MPC=y
-CT_CC_GCC_HAS_LIBQUADMATH=y
-CT_CC_GCC_HAS_LIBSANITIZER=y
-CT_CC_GCC_VERSION="4.8.5"
-#CT_CC_GCC_VERSION="5.2.0"
-# CT_CC_LANG_FORTRAN is not set
+CT_CC_CHOICE_KSYM="GCC"
+CT_CC_GCC_SHOW=y
+
+#
+# Options for gcc
+#
+CT_CC_GCC_PKG_KSYM="GCC"
+CT_GCC_DIR_NAME="gcc"
+CT_GCC_USE_GNU=y
+# CT_GCC_USE_ORACLE is not set
+CT_GCC_USE="GCC"
+CT_GCC_PKG_NAME="gcc"
+CT_GCC_SRC_RELEASE=y
+# CT_GCC_SRC_DEVEL is not set
+CT_GCC_PATCH_ORDER="global"
+CT_GCC_V_13=y
+# CT_GCC_V_12 is not set
+# CT_GCC_V_11 is not set
+# CT_GCC_V_10 is not set
+# CT_GCC_V_9 is not set
+# CT_GCC_V_8 is not set
+# CT_GCC_V_7 is not set
+# CT_GCC_V_6 is not set
+# CT_GCC_V_5 is not set
+# CT_GCC_V_4_9 is not set
+CT_GCC_VERSION="13.2.0"
+CT_GCC_MIRRORS="$(CT_Mirrors GNU gcc/gcc-${CT_GCC_VERSION}) $(CT_Mirrors sourceware gcc/releases/gcc-${CT_GCC_VERSION})"
+CT_GCC_ARCHIVE_FILENAME="@{pkg_name}-@{version}"
+CT_GCC_ARCHIVE_DIRNAME="@{pkg_name}-@{version}"
+CT_GCC_ARCHIVE_FORMATS=".tar.xz .tar.gz"
+CT_GCC_SIGNATURE_FORMAT=""
+CT_GCC_later_than_13=y
+CT_GCC_13_or_later=y
+CT_GCC_later_than_12=y
+CT_GCC_12_or_later=y
+CT_GCC_later_than_11=y
+CT_GCC_11_or_later=y
+CT_GCC_later_than_10=y
+CT_GCC_10_or_later=y
+CT_GCC_later_than_9=y
+CT_GCC_9_or_later=y
+CT_GCC_later_than_8=y
+CT_GCC_8_or_later=y
+CT_GCC_later_than_7=y
+CT_GCC_7_or_later=y
+CT_GCC_later_than_6=y
+CT_GCC_6_or_later=y
+CT_GCC_later_than_5=y
+CT_GCC_5_or_later=y
+CT_GCC_later_than_4_9=y
+CT_GCC_4_9_or_later=y
+CT_GCC_REQUIRE_4_9_or_later=y
+CT_CC_GCC_HAS_LIBMPX=y
 CT_CC_GCC_ENABLE_CXX_FLAGS=""
 CT_CC_GCC_CORE_EXTRA_CONFIG_ARRAY=""
 CT_CC_GCC_EXTRA_CONFIG_ARRAY=""
-CT_CC_GCC_EXTRA_ENV_ARRAY=""
 CT_CC_GCC_STATIC_LIBSTDCXX=y
 # CT_CC_GCC_SYSTEM_ZLIB is not set
+CT_CC_GCC_CONFIG_TLS=m
 
 #
 # Optimisation features
 #
+CT_CC_GCC_USE_GRAPHITE=y
 
 #
 # Settings for libraries running on target
 #
+# CT_CC_GCC_ENABLE_DEFAULT_PIE is not set
 CT_CC_GCC_ENABLE_TARGET_OPTSPACE=y
 # CT_CC_GCC_LIBMUDFLAP is not set
 # CT_CC_GCC_LIBGOMP is not set
 # CT_CC_GCC_LIBSSP is not set
 # CT_CC_GCC_LIBQUADMATH is not set
 # CT_CC_GCC_LIBSANITIZER is not set
+CT_CC_GCC_LIBMPX=y
+CT_CC_GCC_LIBSTDCXX_VERBOSE=m
 
 #
 # Misc. obscure options.
 #
 CT_CC_CXA_ATEXIT=y
+CT_CC_GCC_TM_CLONE_REGISTRY=m
 CT_CC_GCC_DISABLE_PCH=y
 CT_CC_GCC_SJLJ_EXCEPTIONS=m
 CT_CC_GCC_LDBL_128=m
@@ -416,105 +613,256 @@ CT_CC_GCC_LNK_HASH_STYLE_DEFAULT=y
 # CT_CC_GCC_LNK_HASH_STYLE_GNU is not set
 # CT_CC_GCC_LNK_HASH_STYLE_BOTH is not set
 CT_CC_GCC_LNK_HASH_STYLE=""
-CT_CC_GCC_DEC_FLOAT_AUTO=y
-# CT_CC_GCC_DEC_FLOAT_BID is not set
-# CT_CC_GCC_DEC_FLOAT_DPD is not set
+CT_CC_GCC_DEC_FLOATS_AUTO=y
+# CT_CC_GCC_DEC_FLOATS_BID is not set
+# CT_CC_GCC_DEC_FLOATS_DPD is not set
 # CT_CC_GCC_DEC_FLOATS_NO is not set
-CT_CC_SUPPORT_CXX=y
-CT_CC_SUPPORT_FORTRAN=y
-CT_CC_SUPPORT_JAVA=y
-CT_CC_SUPPORT_ADA=y
-CT_CC_SUPPORT_OBJC=y
-CT_CC_SUPPORT_OBJCXX=y
-CT_CC_SUPPORT_GOLANG=y
+CT_CC_GCC_DEC_FLOATS=""
+CT_ALL_CC_CHOICES="GCC"
 
 #
 # Additional supported languages:
 #
 CT_CC_LANG_CXX=y
-# CT_CC_LANG_JAVA is not set
+# CT_CC_LANG_FORTRAN is not set
+# end of C compiler
 
 #
 # Debug facilities
 #
-# CT_DEBUG_dmalloc is not set
-# CT_DEBUG_duma is not set
-# CT_DEBUG_gdb is not set
-# CT_DEBUG_ltrace is not set
-# CT_DEBUG_strace is not set
+# CT_DEBUG_DUMA is not set
+# CT_DEBUG_GDB is not set
+# CT_DEBUG_LTRACE is not set
+# CT_DEBUG_STRACE is not set
+CT_ALL_DEBUG_CHOICES="DUMA GDB LTRACE STRACE"
+# end of Debug facilities
 
 #
 # Companion libraries
 #
-CT_COMPLIBS_NEEDED=y
+# CT_COMPLIBS_CHECK is not set
+# CT_COMP_LIBS_CLOOG is not set
+CT_COMP_LIBS_EXPAT=y
+CT_COMP_LIBS_EXPAT_PKG_KSYM="EXPAT"
+CT_EXPAT_DIR_NAME="expat"
+CT_EXPAT_PKG_NAME="expat"
+CT_EXPAT_SRC_RELEASE=y
+# CT_EXPAT_SRC_DEVEL is not set
+CT_EXPAT_PATCH_ORDER="global"
+CT_EXPAT_V_2_5=y
+CT_EXPAT_VERSION="2.5.0"
+CT_EXPAT_MIRRORS="http://downloads.sourceforge.net/project/expat/expat/${CT_EXPAT_VERSION} [REDACTED_PROFILE_URL]/libexpat/releases/download/R_${CT_EXPAT_VERSION//./_}"
+CT_EXPAT_ARCHIVE_FILENAME="@{pkg_name}-@{version}"
+CT_EXPAT_ARCHIVE_DIRNAME="@{pkg_name}-@{version}"
+CT_EXPAT_ARCHIVE_FORMATS=".tar.xz .tar.lz .tar.bz2 .tar.gz"
+CT_EXPAT_SIGNATURE_FORMAT=""
+CT_COMP_LIBS_GETTEXT=y
+CT_COMP_LIBS_GETTEXT_PKG_KSYM="GETTEXT"
+CT_GETTEXT_DIR_NAME="gettext"
+CT_GETTEXT_PKG_NAME="gettext"
+CT_GETTEXT_SRC_RELEASE=y
+# CT_GETTEXT_SRC_DEVEL is not set
+CT_GETTEXT_PATCH_ORDER="global"
+CT_GETTEXT_V_0_21=y
+# CT_GETTEXT_V_0_20_1 is not set
+# CT_GETTEXT_V_0_19_8_1 is not set
+CT_GETTEXT_VERSION="0.21"
+CT_GETTEXT_MIRRORS="$(CT_Mirrors GNU gettext)"
+CT_GETTEXT_ARCHIVE_FILENAME="@{pkg_name}-@{version}"
+CT_GETTEXT_ARCHIVE_DIRNAME="@{pkg_name}-@{version}"
+CT_GETTEXT_ARCHIVE_FORMATS=".tar.xz .tar.gz"
+CT_GETTEXT_SIGNATURE_FORMAT="packed/.sig"
+CT_GETTEXT_0_21_or_later=y
+CT_GETTEXT_0_21_or_older=y
+CT_GETTEXT_INCOMPATIBLE_WITH_UCLIBC_NG=y
+
+#
+# This version of gettext is not compatible with uClibc-NG. Select
+#
+
+#
+# a different version if uClibc-NG is used on the target or (in a
+#
+
+#
+# Canadian cross build) on the host.
+#
+CT_COMP_LIBS_GMP=y
+CT_COMP_LIBS_GMP_PKG_KSYM="GMP"
+CT_GMP_DIR_NAME="gmp"
+CT_GMP_PKG_NAME="gmp"
+CT_GMP_SRC_RELEASE=y
+# CT_GMP_SRC_DEVEL is not set
+CT_GMP_PATCH_ORDER="global"
+CT_GMP_V_6_2=y
+# CT_GMP_V_6_1 is not set
+CT_GMP_VERSION="6.2.1"
+CT_GMP_MIRRORS="https://gmplib.org/download/gmp https://gmplib.org/download/gmp/archive $(CT_Mirrors GNU gmp)"
+CT_GMP_ARCHIVE_FILENAME="@{pkg_name}-@{version}"
+CT_GMP_ARCHIVE_DIRNAME="@{pkg_name}-@{version}"
+CT_GMP_ARCHIVE_FORMATS=".tar.xz .tar.lz .tar.bz2"
+CT_GMP_SIGNATURE_FORMAT="packed/.sig"
+CT_COMP_LIBS_ISL=y
+CT_COMP_LIBS_ISL_PKG_KSYM="ISL"
+CT_ISL_DIR_NAME="isl"
+CT_ISL_PKG_NAME="isl"
+CT_ISL_SRC_RELEASE=y
+# CT_ISL_SRC_DEVEL is not set
+CT_ISL_PATCH_ORDER="global"
+CT_ISL_V_0_26=y
+# CT_ISL_V_0_25 is not set
+# CT_ISL_V_0_24 is not set
+# CT_ISL_V_0_23 is not set
+# CT_ISL_V_0_22 is not set
+# CT_ISL_V_0_21 is not set
+# CT_ISL_V_0_20 is not set
+# CT_ISL_V_0_19 is not set
+# CT_ISL_V_0_18 is not set
+# CT_ISL_V_0_17 is not set
+# CT_ISL_V_0_16 is not set
+# CT_ISL_V_0_15 is not set
+CT_ISL_VERSION="0.26"
+CT_ISL_MIRRORS="https://libisl.sourceforge.io"
+CT_ISL_ARCHIVE_FILENAME="@{pkg_name}-@{version}"
+CT_ISL_ARCHIVE_DIRNAME="@{pkg_name}-@{version}"
+CT_ISL_ARCHIVE_FORMATS=".tar.xz .tar.bz2 .tar.gz"
+CT_ISL_SIGNATURE_FORMAT=""
+CT_ISL_later_than_0_18=y
+CT_ISL_0_18_or_later=y
+CT_ISL_later_than_0_15=y
+CT_ISL_0_15_or_later=y
+# CT_COMP_LIBS_LIBELF is not set
+CT_COMP_LIBS_LIBICONV=y
+CT_COMP_LIBS_LIBICONV_PKG_KSYM="LIBICONV"
+CT_LIBICONV_DIR_NAME="libiconv"
+CT_LIBICONV_PKG_NAME="libiconv"
+CT_LIBICONV_SRC_RELEASE=y
+# CT_LIBICONV_SRC_DEVEL is not set
+CT_LIBICONV_PATCH_ORDER="global"
+CT_LIBICONV_V_1_16=y
+# CT_LIBICONV_V_1_15 is not set
+CT_LIBICONV_VERSION="1.16"
+CT_LIBICONV_MIRRORS="$(CT_Mirrors GNU libiconv)"
+CT_LIBICONV_ARCHIVE_FILENAME="@{pkg_name}-@{version}"
+CT_LIBICONV_ARCHIVE_DIRNAME="@{pkg_name}-@{version}"
+CT_LIBICONV_ARCHIVE_FORMATS=".tar.gz"
+CT_LIBICONV_SIGNATURE_FORMAT="packed/.sig"
+CT_COMP_LIBS_MPC=y
+CT_COMP_LIBS_MPC_PKG_KSYM="MPC"
+CT_MPC_DIR_NAME="mpc"
+CT_MPC_PKG_NAME="mpc"
+CT_MPC_SRC_RELEASE=y
+# CT_MPC_SRC_DEVEL is not set
+CT_MPC_PATCH_ORDER="global"
+CT_MPC_V_1_2=y
+CT_MPC_VERSION="1.2.1"
+CT_MPC_MIRRORS="https://www.multiprecision.org/downloads $(CT_Mirrors GNU mpc)"
+CT_MPC_ARCHIVE_FILENAME="@{pkg_name}-@{version}"
+CT_MPC_ARCHIVE_DIRNAME="@{pkg_name}-@{version}"
+CT_MPC_ARCHIVE_FORMATS=".tar.gz"
+CT_MPC_SIGNATURE_FORMAT="packed/.sig"
+CT_MPC_later_than_1_1_0=y
+CT_MPC_1_1_0_or_later=y
+CT_COMP_LIBS_MPFR=y
+CT_COMP_LIBS_MPFR_PKG_KSYM="MPFR"
+CT_MPFR_DIR_NAME="mpfr"
+CT_MPFR_PKG_NAME="mpfr"
+CT_MPFR_SRC_RELEASE=y
+# CT_MPFR_SRC_DEVEL is not set
+CT_MPFR_PATCH_ORDER="global"
+CT_MPFR_V_4_2=y
+CT_MPFR_VERSION="4.2.1"
+CT_MPFR_MIRRORS="https://www.mpfr.org/mpfr-${CT_MPFR_VERSION} $(CT_Mirrors GNU mpfr)"
+CT_MPFR_ARCHIVE_FILENAME="@{pkg_name}-@{version}"
+CT_MPFR_ARCHIVE_DIRNAME="@{pkg_name}-@{version}"
+CT_MPFR_ARCHIVE_FORMATS=".tar.xz .tar.bz2 .tar.gz .zip"
+CT_MPFR_SIGNATURE_FORMAT="packed/.asc"
+CT_MPFR_later_than_4_0_0=y
+CT_MPFR_4_0_0_or_later=y
+CT_COMP_LIBS_NCURSES=y
+CT_COMP_LIBS_NCURSES_PKG_KSYM="NCURSES"
+CT_NCURSES_DIR_NAME="ncurses"
+CT_NCURSES_PKG_NAME="ncurses"
+CT_NCURSES_SRC_RELEASE=y
+# CT_NCURSES_SRC_DEVEL is not set
+CT_NCURSES_PATCH_ORDER="global"
+CT_NCURSES_V_6_4=y
+# CT_NCURSES_V_6_2 is not set
+# CT_NCURSES_V_6_1 is not set
+# CT_NCURSES_V_6_0 is not set
+CT_NCURSES_VERSION="6.4"
+CT_NCURSES_MIRRORS="https://invisible-mirror.net/archives/ncurses $(CT_Mirrors GNU ncurses)"
+CT_NCURSES_ARCHIVE_FILENAME="@{pkg_name}-@{version}"
+CT_NCURSES_ARCHIVE_DIRNAME="@{pkg_name}-@{version}"
+CT_NCURSES_ARCHIVE_FORMATS=".tar.gz"
+CT_NCURSES_SIGNATURE_FORMAT="packed/.sig"
+CT_NCURSES_NEW_ABI=y
+CT_NCURSES_HOST_CONFIG_ARGS=""
+CT_NCURSES_HOST_DISABLE_DB=y
+CT_NCURSES_HOST_FALLBACKS="linux,xterm,xterm-color,xterm-256color,vt100"
+CT_NCURSES_TARGET_CONFIG_ARGS=""
+# CT_NCURSES_TARGET_DISABLE_DB is not set
+CT_NCURSES_TARGET_FALLBACKS=""
+CT_COMP_LIBS_ZLIB=y
+CT_COMP_LIBS_ZLIB_PKG_KSYM="ZLIB"
+CT_ZLIB_DIR_NAME="zlib"
+CT_ZLIB_PKG_NAME="zlib"
+CT_ZLIB_SRC_RELEASE=y
+# CT_ZLIB_SRC_DEVEL is not set
+CT_ZLIB_PATCH_ORDER="global"
+CT_ZLIB_V_1_2_13=y
+CT_ZLIB_VERSION="1.2.13"
+CT_ZLIB_MIRRORS="[REDACTED_PROFILE_URL]/zlib/releases/download/v${CT_ZLIB_VERSION} https://www.zlib.net/"
+CT_ZLIB_ARCHIVE_FILENAME="@{pkg_name}-@{version}"
+CT_ZLIB_ARCHIVE_DIRNAME="@{pkg_name}-@{version}"
+CT_ZLIB_ARCHIVE_FORMATS=".tar.xz .tar.gz"
+CT_ZLIB_SIGNATURE_FORMAT="packed/.asc"
+CT_COMP_LIBS_ZSTD=y
+CT_COMP_LIBS_ZSTD_PKG_KSYM="ZSTD"
+CT_ZSTD_DIR_NAME="zstd"
+CT_ZSTD_PKG_NAME="zstd"
+CT_ZSTD_SRC_RELEASE=y
+# CT_ZSTD_SRC_DEVEL is not set
+CT_ZSTD_PATCH_ORDER="global"
+CT_ZSTD_V_1_5_5=y
+# CT_ZSTD_V_1_5_2 is not set
+CT_ZSTD_VERSION="1.5.5"
+CT_ZSTD_MIRRORS="[REDACTED_PROFILE_URL]/zstd/releases/download/v${CT_ZSTD_VERSION} https://www.zstd.net/"
+CT_ZSTD_ARCHIVE_FILENAME="@{pkg_name}-@{version}"
+CT_ZSTD_ARCHIVE_DIRNAME="@{pkg_name}-@{version}"
+CT_ZSTD_ARCHIVE_FORMATS=".tar.gz"
+CT_ZSTD_SIGNATURE_FORMAT="packed/.sig"
+CT_ALL_COMP_LIBS_CHOICES="CLOOG EXPAT GETTEXT GMP GNUPRUMCU ISL LIBELF LIBICONV MPC MPFR NCURSES NEWLIB_NANO PICOLIBC ZLIB ZSTD"
 CT_LIBICONV_NEEDED=y
 CT_GETTEXT_NEEDED=y
 CT_GMP_NEEDED=y
 CT_MPFR_NEEDED=y
 CT_ISL_NEEDED=y
-CT_CLOOG_NEEDED=y
 CT_MPC_NEEDED=y
-CT_COMPLIBS=y
+CT_NCURSES_NEEDED=y
+CT_ZLIB_NEEDED=y
+CT_ZSTD_NEEDED=y
 CT_LIBICONV=y
 CT_GETTEXT=y
 CT_GMP=y
 CT_MPFR=y
 CT_ISL=y
-CT_CLOOG=y
 CT_MPC=y
-CT_LIBICONV_V_1_14=y
-CT_LIBICONV_VERSION="1.14"
-CT_GETTEXT_V_0_19_6=y
-CT_GETTEXT_VERSION="0.19.6"
-CT_GMP_V_6_0_0=y
-# CT_GMP_V_5_1_3 is not set
-# CT_GMP_V_5_1_1 is not set
-# CT_GMP_V_5_0_2 is not set
-# CT_GMP_V_5_0_1 is not set
-# CT_GMP_V_4_3_2 is not set
-# CT_GMP_V_4_3_1 is not set
-# CT_GMP_V_4_3_0 is not set
-CT_GMP_5_0_2_or_later=y
-CT_GMP_VERSION="6.0.0a"
-CT_MPFR_V_3_1_3=y
-# CT_MPFR_V_3_1_2 is not set
-# CT_MPFR_V_3_1_0 is not set
-# CT_MPFR_V_3_0_1 is not set
-# CT_MPFR_V_3_0_0 is not set
-# CT_MPFR_V_2_4_2 is not set
-# CT_MPFR_V_2_4_1 is not set
-# CT_MPFR_V_2_4_0 is not set
-CT_MPFR_VERSION="3.1.3"
-CT_ISL_V_0_14=y
-CT_ISL_V_0_14_or_later=y
-CT_ISL_V_0_12_or_later=y
-CT_ISL_VERSION="0.14"
-CT_CLOOG_V_0_18_4=y
-# CT_CLOOG_V_0_18_1 is not set
-# CT_CLOOG_V_0_18_0 is not set
-CT_CLOOG_VERSION="0.18.4"
-CT_CLOOG_0_18_4_or_later=y
-CT_CLOOG_0_18_or_later=y
-CT_MPC_V_1_0_3=y
-# CT_MPC_V_1_0_2 is not set
-# CT_MPC_V_1_0_1 is not set
-# CT_MPC_V_1_0 is not set
-# CT_MPC_V_0_9 is not set
-# CT_MPC_V_0_8_2 is not set
-# CT_MPC_V_0_8_1 is not set
-# CT_MPC_V_0_7 is not set
-CT_MPC_VERSION="1.0.3"
-
-#
-# Companion libraries common options
-#
-# CT_COMPLIBS_CHECK is not set
+CT_NCURSES=y
+CT_ZLIB=y
+CT_ZSTD=y
+# end of Companion libraries
 
 #
 # Companion tools
 #
-
-#
-# READ HELP before you say 'Y' below !!!
-#
-# CT_COMP_TOOLS is not set
+# CT_COMP_TOOLS_FOR_HOST is not set
+# CT_COMP_TOOLS_AUTOCONF is not set
+# CT_COMP_TOOLS_AUTOMAKE is not set
+# CT_COMP_TOOLS_BISON is not set
+# CT_COMP_TOOLS_DTC is not set
+# CT_COMP_TOOLS_LIBTOOL is not set
+# CT_COMP_TOOLS_M4 is not set
+# CT_COMP_TOOLS_MAKE is not set
+CT_ALL_COMP_TOOLS_CHOICES="AUTOCONF AUTOMAKE BISON DTC LIBTOOL M4 MAKE"
+# end of Companion tools
diff --git a/Engine/Build/BatchFiles/Linux/Toolchain/x86_64-unknown-linux-gnu.windows.config b/Engine/Build/BatchFiles/Linux/Toolchain/x86_64-unknown-linux-gnu.windows.config
--- a/Engine/Build/BatchFiles/Linux/Toolchain/x86_64-unknown-linux-gnu.windows.config
+++ b/Engine/Build/BatchFiles/Linux/Toolchain/x86_64-unknown-linux-gnu.windows.config
@@ -1,10 +1,34 @@
 #
 # Automatically generated file; DO NOT EDIT.
-# Crosstool-NG Configuration
-#
-CT_CONFIGURE_has_make381=y
-CT_CONFIGURE_has_xz=y
-CT_CONFIGURE_has_cvs=y
+# crosstool-NG 1.26.0 Configuration
+#
+CT_CONFIGURE_has_static_link=y
+CT_CONFIGURE_has_cxx11=y
+CT_CONFIGURE_has_wget=y
+CT_CONFIGURE_has_curl=y
+CT_CONFIGURE_has_rsync=y
+CT_CONFIGURE_has_make_3_81_or_newer=y
+CT_CONFIGURE_has_make_4_0_or_newer=y
+CT_CONFIGURE_has_libtool_2_4_or_newer=y
+CT_CONFIGURE_has_libtoolize_2_4_or_newer=y
+CT_CONFIGURE_has_autoconf_2_65_or_newer=y
+CT_CONFIGURE_has_autoreconf_2_65_or_newer=y
+CT_CONFIGURE_has_automake_1_15_or_newer=y
+CT_CONFIGURE_has_gnu_m4_1_4_12_or_newer=y
+CT_CONFIGURE_has_bison_2_7_or_newer=y
+CT_CONFIGURE_has_bison_3_0_4_or_newer=y
+CT_CONFIGURE_has_python=y
+CT_CONFIGURE_has_git=y
+CT_CONFIGURE_has_md5sum=y
+CT_CONFIGURE_has_sha1sum=y
+CT_CONFIGURE_has_sha256sum=y
+CT_CONFIGURE_has_sha512sum=y
+CT_CONFIGURE_has_install_with_strip_program=y
+CT_VERSION="1.26.0"
+CT_VCHECK=""
+CT_CONFIG_VERSION_ENV="4"
+CT_CONFIG_VERSION_CURRENT="4"
+CT_CONFIG_VERSION="4"
 CT_MODULES=y
 
 #
@@ -21,41 +45,49 @@ CT_MODULES=y
 #
 # Paths
 #
-CT_LOCAL_TARBALLS_DIR=""
+CT_LOCAL_TARBALLS_DIR="${HOME}/src"
+CT_SAVE_TARBALLS=y
+# CT_TARBALLS_BUILDROOT_LAYOUT is not set
 CT_WORK_DIR="${CT_TOP_DIR}/.build"
+CT_BUILD_TOP_DIR="${CT_WORK_DIR:-${CT_TOP_DIR}/.build}/${CT_HOST:+HOST-${CT_HOST}/}${CT_TARGET}"
+CT_BUILD_DIR="${CT_BUILD_TOP_DIR}/build"
 CT_PREFIX_DIR="${CT_TOP_DIR}/../OUTPUT-windows/${CT_TARGET}"
-CT_INSTALL_DIR="${CT_PREFIX_DIR}"
 CT_RM_RF_PREFIX_DIR=y
 CT_REMOVE_DOCS=y
-CT_INSTALL_DIR_RO=y
+CT_INSTALL_LICENSES=y
+CT_PREFIX_DIR_RO=y
 CT_STRIP_HOST_TOOLCHAIN_EXECUTABLES=y
 # CT_STRIP_TARGET_TOOLCHAIN_EXECUTABLES is not set
 
 #
 # Downloading
 #
+CT_DOWNLOAD_AGENT_WGET=y
+# CT_DOWNLOAD_AGENT_CURL is not set
+# CT_DOWNLOAD_AGENT_NONE is not set
 # CT_FORBID_DOWNLOAD is not set
 # CT_FORCE_DOWNLOAD is not set
 CT_CONNECT_TIMEOUT=10
+CT_DOWNLOAD_WGET_OPTIONS="--passive-ftp --tries=3 -nc --progress=dot:binary"
 # CT_ONLY_DOWNLOAD is not set
 # CT_USE_MIRROR is not set
+CT_VERIFY_DOWNLOAD_DIGEST=y
+CT_VERIFY_DOWNLOAD_DIGEST_SHA512=y
+# CT_VERIFY_DOWNLOAD_DIGEST_SHA256 is not set
+# CT_VERIFY_DOWNLOAD_DIGEST_SHA1 is not set
+# CT_VERIFY_DOWNLOAD_DIGEST_MD5 is not set
+CT_VERIFY_DOWNLOAD_DIGEST_ALG="sha512"
+# CT_VERIFY_DOWNLOAD_SIGNATURE is not set
 
 #
 # Extracting
 #
 # CT_FORCE_EXTRACT is not set
-CT_OVERIDE_CONFIG_GUESS_SUB=y
+CT_OVERRIDE_CONFIG_GUESS_SUB=y
 # CT_ONLY_EXTRACT is not set
-# CT_PATCH_BUNDLED is not set
-# CT_PATCH_LOCAL is not set
-CT_PATCH_BUNDLED_LOCAL=y
-# CT_PATCH_LOCAL_BUNDLED is not set
-# CT_PATCH_BUNDLED_FALLBACK_LOCAL is not set
-# CT_PATCH_LOCAL_FALLBACK_BUNDLED is not set
-# CT_PATCH_NONE is not set
-CT_PATCH_ORDER="bundled,local"
-CT_PATCH_USE_LOCAL=y
-CT_LOCAL_PATCH_DIR="${CT_TOP_DIR}/../../patches"
+CT_PATCH_BUNDLED=y
+# CT_PATCH_BUNDLED_LOCAL is not set
+CT_PATCH_ORDER="bundled"
 
 #
 # Build behavior
@@ -64,8 +96,9 @@ CT_PARALLEL_JOBS=0
 CT_LOAD=""
 CT_USE_PIPES=y
 CT_EXTRA_CFLAGS_FOR_BUILD=""
+CT_EXTRA_CXXFLAGS_FOR_BUILD=""
 CT_EXTRA_LDFLAGS_FOR_BUILD=""
-CT_EXTRA_CFLAGS_FOR_HOST="-fgnu89-inline"
+CT_EXTRA_CFLAGS_FOR_HOST=""
 CT_EXTRA_LDFLAGS_FOR_HOST=""
 # CT_CONFIG_SHELL_SH is not set
 # CT_CONFIG_SHELL_ASH is not set
@@ -78,71 +111,74 @@ CT_CONFIG_SHELL="${bash}"
 #
 # CT_LOG_ERROR is not set
 # CT_LOG_WARN is not set
-CT_LOG_INFO=y
-# CT_LOG_EXTRA is not set
+# CT_LOG_INFO is not set
+CT_LOG_EXTRA=y
 # CT_LOG_ALL is not set
 # CT_LOG_DEBUG is not set
-CT_LOG_LEVEL_MAX="INFO"
+CT_LOG_LEVEL_MAX="EXTRA"
 # CT_LOG_SEE_TOOLS_WARN is not set
 CT_LOG_PROGRESS_BAR=y
 CT_LOG_TO_FILE=y
 CT_LOG_FILE_COMPRESS=y
+# end of Paths and misc options
 
 #
 # Target options
 #
+# CT_ARCH_ALPHA is not set
+# CT_ARCH_ARC is not set
+# CT_ARCH_ARM is not set
+# CT_ARCH_AVR is not set
+# CT_ARCH_BPF is not set
+# CT_ARCH_M68K is not set
+# CT_ARCH_MIPS is not set
+# CT_ARCH_NIOS2 is not set
+# CT_ARCH_POWERPC is not set
+# CT_ARCH_PRU is not set
+# CT_ARCH_S390 is not set
+# CT_ARCH_SH is not set
+# CT_ARCH_SPARC is not set
+CT_ARCH_X86=y
+# CT_ARCH_XTENSA is not set
 CT_ARCH="x86"
-CT_ARCH_SUPPORTS_32=y
-CT_ARCH_SUPPORTS_64=y
-CT_ARCH_SUPPORTS_WITH_ARCH=y
-CT_ARCH_SUPPORTS_WITH_CPU=y
-CT_ARCH_SUPPORTS_WITH_TUNE=y
-CT_ARCH_DEFAULT_32=y
-CT_ARCH_ARCH=""
+CT_ARCH_CHOICE_KSYM="X86"
 CT_ARCH_CPU=""
 CT_ARCH_TUNE=""
-# CT_ARCH_32 is not set
-CT_ARCH_64=y
-CT_ARCH_BITNESS=64
-CT_TARGET_CFLAGS=""
-CT_TARGET_LDFLAGS=""
-# CT_ARCH_alpha is not set
-# CT_ARCH_arm is not set
-# CT_ARCH_avr is not set
-# CT_ARCH_m68k is not set
-# CT_ARCH_mips is not set
-# CT_ARCH_nios2 is not set
-# CT_ARCH_powerpc is not set
-# CT_ARCH_s390 is not set
-# CT_ARCH_sh is not set
-# CT_ARCH_sparc is not set
-CT_ARCH_x86=y
-# CT_ARCH_xtensa is not set
-CT_ARCH_alpha_AVAILABLE=y
-CT_ARCH_arm_AVAILABLE=y
-CT_ARCH_avr_AVAILABLE=y
-CT_ARCH_m68k_AVAILABLE=y
-CT_ARCH_microblaze_AVAILABLE=y
-CT_ARCH_mips_AVAILABLE=y
-CT_ARCH_nios2_AVAILABLE=y
-CT_ARCH_powerpc_AVAILABLE=y
-CT_ARCH_s390_AVAILABLE=y
-CT_ARCH_sh_AVAILABLE=y
-CT_ARCH_sparc_AVAILABLE=y
-CT_ARCH_x86_AVAILABLE=y
-CT_ARCH_xtensa_AVAILABLE=y
+CT_ARCH_X86_SHOW=y
+
+#
+# Options for x86
+#
+CT_ARCH_X86_PKG_KSYM=""
+CT_ALL_ARCH_CHOICES="ALPHA ARC ARM AVR BPF C6X LOONGARCH M68K MICROBLAZE MIPS MOXIE MSP430 NIOS2 POWERPC PRU RISCV S390 SH SPARC X86 XTENSA"
 CT_ARCH_SUFFIX=""
+# CT_OMIT_TARGET_VENDOR is not set
 
 #
 # Generic target options
 #
 # CT_MULTILIB is not set
+CT_DEMULTILIB=y
 CT_ARCH_USE_MMU=y
+CT_ARCH_SUPPORTS_LIBSANITIZER=y
+CT_ARCH_SUPPORTS_32=y
+CT_ARCH_SUPPORTS_64=y
+CT_ARCH_DEFAULT_32=y
+CT_ARCH_BITNESS=64
+# CT_ARCH_32 is not set
+CT_ARCH_64=y
+CT_ARCH_SUPPORTS_WITH_32_64=y
 
 #
 # Target optimisations
 #
-CT_ARCH_FLOAT=""
+CT_ARCH_SUPPORTS_WITH_ARCH=y
+CT_ARCH_SUPPORTS_WITH_CPU=y
+CT_ARCH_SUPPORTS_WITH_TUNE=y
+CT_ARCH_ARCH=""
+CT_TARGET_CFLAGS=""
+CT_TARGET_LDFLAGS=""
+# end of Target options
 
 #
 # Toolchain options
@@ -155,8 +191,8 @@ CT_FORCE_SYSROOT=y
 CT_USE_SYSROOT=y
 CT_SYSROOT_NAME="sysroot"
 CT_SYSROOT_DIR_PREFIX=""
-CT_WANTS_STATIC_LINK=y
 CT_STATIC_TOOLCHAIN=y
+CT_SHOW_CT_VERSION=y
 CT_TOOLCHAIN_PKGVERSION=""
 CT_TOOLCHAIN_BUGURL=""
 
@@ -192,132 +228,279 @@ CT_HOST_SUFFIX=""
 # Misc options
 #
 # CT_TOOLCHAIN_ENABLE_NLS is not set
+# end of Toolchain options
 
 #
 # Operating System
 #
 CT_KERNEL_SUPPORTS_SHARED_LIBS=y
+# CT_KERNEL_BARE_METAL is not set
+CT_KERNEL_LINUX=y
 CT_KERNEL="linux"
-CT_KERNEL_VERSION="3.10.93"
-# CT_KERNEL_bare_metal is not set
-CT_KERNEL_linux=y
-CT_KERNEL_bare_metal_AVAILABLE=y
-CT_KERNEL_linux_AVAILABLE=y
-# CT_KERNEL_V_4_3 is not set
-# CT_KERNEL_V_4_2 is not set
-# CT_KERNEL_V_4_1 is not set
-# CT_KERNEL_V_3_18 is not set
-# CT_KERNEL_V_3_14 is not set
-# CT_KERNEL_V_3_12 is not set
-CT_KERNEL_V_3_10=y
-# CT_KERNEL_V_3_4 is not set
-# CT_KERNEL_V_3_2 is not set
-# CT_KERNEL_V_2_6_32 is not set
-# CT_KERNEL_LINUX_CUSTOM is not set
-CT_KERNEL_windows_AVAILABLE=y
-
-#
-# Common kernel options
-#
-CT_SHARED_LIBS=y
-
-#
-# linux other options
-#
+CT_KERNEL_CHOICE_KSYM="LINUX"
+CT_KERNEL_LINUX_SHOW=y
+
+#
+# Options for linux
+#
+CT_KERNEL_LINUX_PKG_KSYM="LINUX"
+CT_LINUX_DIR_NAME="linux"
+CT_LINUX_USE_WWW_KERNEL_ORG=y
+# CT_LINUX_USE_ORACLE is not set
+CT_LINUX_USE="LINUX"
+CT_LINUX_PKG_NAME="linux"
+CT_LINUX_SRC_RELEASE=y
+# CT_LINUX_SRC_DEVEL is not set
+CT_LINUX_PATCH_ORDER="global"
+# CT_LINUX_V_6_4 is not set
+# CT_LINUX_V_6_3 is not set
+# CT_LINUX_V_6_2 is not set
+# CT_LINUX_V_6_1 is not set
+# CT_LINUX_V_6_0 is not set
+# CT_LINUX_V_5_19 is not set
+# CT_LINUX_V_5_18 is not set
+# CT_LINUX_V_5_17 is not set
+# CT_LINUX_V_5_16 is not set
+# CT_LINUX_V_5_15 is not set
+# CT_LINUX_V_5_14 is not set
+# CT_LINUX_V_5_13 is not set
+# CT_LINUX_V_5_12 is not set
+# CT_LINUX_V_5_11 is not set
+# CT_LINUX_V_5_10 is not set
+# CT_LINUX_V_5_9 is not set
+# CT_LINUX_V_5_8 is not set
+# CT_LINUX_V_5_7 is not set
+# CT_LINUX_V_5_4 is not set
+# CT_LINUX_V_5_3 is not set
+# CT_LINUX_V_5_2 is not set
+# CT_LINUX_V_5_1 is not set
+# CT_LINUX_V_5_0 is not set
+# CT_LINUX_V_4_20 is not set
+# CT_LINUX_V_4_19 is not set
+CT_LINUX_V_4_18=y
+# CT_LINUX_V_4_17 is not set
+# CT_LINUX_V_4_16 is not set
+# CT_LINUX_V_4_15 is not set
+# CT_LINUX_V_4_14 is not set
+# CT_LINUX_V_4_13 is not set
+# CT_LINUX_V_4_12 is not set
+# CT_LINUX_V_4_11 is not set
+# CT_LINUX_V_4_10 is not set
+# CT_LINUX_V_4_9 is not set
+# CT_LINUX_V_4_4 is not set
+# CT_LINUX_V_4_1 is not set
+# CT_LINUX_V_3_16 is not set
+# CT_LINUX_V_3_13 is not set
+# CT_LINUX_V_3_12 is not set
+# CT_LINUX_V_3_10 is not set
+# CT_LINUX_V_3_4 is not set
+# CT_LINUX_V_3_2 is not set
+CT_LINUX_VERSION="4.18.20"
+CT_LINUX_MIRRORS="$(CT_Mirrors kernel.org linux ${CT_LINUX_VERSION})"
+CT_LINUX_ARCHIVE_FILENAME="@{pkg_name}-@{version}"
+CT_LINUX_ARCHIVE_DIRNAME="@{pkg_name}-@{version}"
+CT_LINUX_ARCHIVE_FORMATS=".tar.xz .tar.gz"
+CT_LINUX_SIGNATURE_FORMAT="unpacked/.sign"
+CT_LINUX_5_19_or_older=y
+CT_LINUX_older_than_5_19=y
+CT_LINUX_5_12_or_older=y
+CT_LINUX_older_than_5_12=y
+CT_LINUX_5_5_or_older=y
+CT_LINUX_older_than_5_5=y
+CT_LINUX_5_3_or_older=y
+CT_LINUX_older_than_5_3=y
+CT_LINUX_later_than_4_8=y
+CT_LINUX_4_8_or_later=y
+CT_LINUX_later_than_3_7=y
+CT_LINUX_3_7_or_later=y
+CT_LINUX_later_than_3_2=y
+CT_LINUX_3_2_or_later=y
+CT_LINUX_REQUIRE_3_2_or_later=y
+CT_KERNEL_DEP_RSYNC=y
 CT_KERNEL_LINUX_VERBOSITY_0=y
 # CT_KERNEL_LINUX_VERBOSITY_1 is not set
 # CT_KERNEL_LINUX_VERBOSITY_2 is not set
 CT_KERNEL_LINUX_VERBOSE_LEVEL=0
 CT_KERNEL_LINUX_INSTALL_CHECK=y
+CT_ALL_KERNEL_CHOICES="BARE_METAL LINUX WINDOWS"
+
+#
+# Common kernel options
+#
+CT_SHARED_LIBS=y
+# end of Operating System
 
 #
 # Binary utilities
 #
 CT_ARCH_BINFMT_ELF=y
+CT_BINUTILS_BINUTILS=y
 CT_BINUTILS="binutils"
-CT_BINUTILS_binutils=y
+CT_BINUTILS_CHOICE_KSYM="BINUTILS"
+CT_BINUTILS_BINUTILS_SHOW=y
+
+#
+# Options for binutils
+#
+CT_BINUTILS_BINUTILS_PKG_KSYM="BINUTILS"
+CT_BINUTILS_DIR_NAME="binutils"
+CT_BINUTILS_USE_GNU=y
+# CT_BINUTILS_USE_ORACLE is not set
+CT_BINUTILS_USE="BINUTILS"
+CT_BINUTILS_PKG_NAME="binutils"
+CT_BINUTILS_SRC_RELEASE=y
+# CT_BINUTILS_SRC_DEVEL is not set
+CT_BINUTILS_PATCH_ORDER="global"
+CT_BINUTILS_V_2_40=y
+# CT_BINUTILS_V_2_39 is not set
+# CT_BINUTILS_V_2_38 is not set
+# CT_BINUTILS_V_2_37 is not set
+# CT_BINUTILS_V_2_36 is not set
+# CT_BINUTILS_V_2_35 is not set
+# CT_BINUTILS_V_2_34 is not set
+# CT_BINUTILS_V_2_33 is not set
+# CT_BINUTILS_V_2_32 is not set
+# CT_BINUTILS_V_2_31 is not set
+# CT_BINUTILS_V_2_30 is not set
+# CT_BINUTILS_V_2_29 is not set
+# CT_BINUTILS_V_2_28 is not set
+# CT_BINUTILS_V_2_27 is not set
+# CT_BINUTILS_V_2_26 is not set
+CT_BINUTILS_VERSION="2.40"
+CT_BINUTILS_MIRRORS="$(CT_Mirrors GNU binutils) $(CT_Mirrors sourceware binutils/releases)"
+CT_BINUTILS_ARCHIVE_FILENAME="@{pkg_name}-@{version}"
+CT_BINUTILS_ARCHIVE_DIRNAME="@{pkg_name}-@{version}"
+CT_BINUTILS_ARCHIVE_FORMATS=".tar.xz .tar.bz2 .tar.gz"
+CT_BINUTILS_SIGNATURE_FORMAT="packed/.sig"
+CT_BINUTILS_later_than_2_39=y
+CT_BINUTILS_2_39_or_later=y
+CT_BINUTILS_later_than_2_30=y
+CT_BINUTILS_2_30_or_later=y
+CT_BINUTILS_later_than_2_27=y
+CT_BINUTILS_2_27_or_later=y
+CT_BINUTILS_later_than_2_26=y
+CT_BINUTILS_2_26_or_later=y
 
 #
 # GNU binutils
 #
-# CT_CC_BINUTILS_SHOW_LINARO is not set
-CT_BINUTILS_V_2_31_1=y
-# CT_BINUTILS_V_2_25 is not set
-# CT_BINUTILS_V_2_24 is not set
-# CT_BINUTILS_V_2_23_2 is not set
-# CT_BINUTILS_V_2_23_1 is not set
-# CT_BINUTILS_V_2_22 is not set
-# CT_BINUTILS_V_2_21_53 is not set
-# CT_BINUTILS_V_2_21_1a is not set
-# CT_BINUTILS_V_2_20_1a is not set
-# CT_BINUTILS_V_2_19_1a is not set
-# CT_BINUTILS_V_2_18a is not set
-CT_BINUTILS_VERSION="2.31.1"
-CT_BINUTILS_2_25_1_or_later=y
-CT_BINUTILS_2_25_or_later=y
-CT_BINUTILS_2_24_or_later=y
-CT_BINUTILS_2_23_or_later=y
-CT_BINUTILS_2_22_or_later=y
-CT_BINUTILS_2_21_or_later=y
-CT_BINUTILS_2_20_or_later=y
-CT_BINUTILS_2_19_or_later=y
-CT_BINUTILS_2_18_or_later=y
-CT_BINUTILS_HAS_HASH_STYLE=y
-CT_BINUTILS_HAS_GOLD=y
 CT_BINUTILS_GOLD_SUPPORTS_ARCH=y
-CT_BINUTILS_GOLD_SUPPORT=y
-CT_BINUTILS_HAS_PLUGINS=y
-CT_BINUTILS_HAS_PKGVERSION_BUGURL=y
-CT_BINUTILS_FORCE_LD_BFD=y
+CT_BINUTILS_FORCE_LD_BFD_DEFAULT=y
 CT_BINUTILS_LINKER_LD=y
-# CT_BINUTILS_LINKER_LD_GOLD is not set
-# CT_BINUTILS_LINKER_GOLD_LD is not set
 CT_BINUTILS_LINKERS_LIST="ld"
 CT_BINUTILS_LINKER_DEFAULT="bfd"
-# CT_BINUTILS_PLUGINS is not set
+CT_BINUTILS_RELRO=m
+CT_BINUTILS_DETERMINISTIC_ARCHIVES=y
 CT_BINUTILS_EXTRA_CONFIG_ARRAY=""
 # CT_BINUTILS_FOR_TARGET is not set
-
-#
-# binutils other options
-#
+CT_ALL_BINUTILS_CHOICES="BINUTILS"
+# end of Binary utilities
 
 #
 # C-library
 #
+CT_LIBC_GLIBC=y
+# CT_LIBC_UCLIBC_NG is not set
 CT_LIBC="glibc"
-CT_LIBC_VERSION="2.17"
-CT_LIBC_glibc=y
-# CT_LIBC_musl is not set
-# CT_LIBC_uClibc is not set
-CT_LIBC_avr_libc_AVAILABLE=y
-CT_LIBC_glibc_AVAILABLE=y
+CT_LIBC_CHOICE_KSYM="GLIBC"
+CT_LIBC_GLIBC_SHOW=y
+
+#
+# Options for glibc
+#
+CT_LIBC_GLIBC_PKG_KSYM="GLIBC"
+CT_GLIBC_DIR_NAME="glibc"
+CT_GLIBC_USE_GNU=y
+# CT_GLIBC_USE_ORACLE is not set
+CT_GLIBC_USE="GLIBC"
+CT_GLIBC_PKG_NAME="glibc"
+CT_GLIBC_SRC_RELEASE=y
+# CT_GLIBC_SRC_DEVEL is not set
+CT_GLIBC_PATCH_ORDER="global"
+CT_GLIBC_V_2_28=y
+# CT_GLIBC_V_2_27 is not set
+# CT_GLIBC_V_2_26 is not set
+# CT_GLIBC_V_2_25 is not set
+# CT_GLIBC_V_2_24 is not set
+# CT_GLIBC_V_2_23 is not set
+# CT_GLIBC_V_2_19 is not set
+# CT_GLIBC_V_2_17 is not set
+CT_GLIBC_VERSION="2.28"
+CT_GLIBC_MIRRORS="$(CT_Mirrors GNU glibc)"
+CT_GLIBC_ARCHIVE_FILENAME="@{pkg_name}-@{version}"
+CT_GLIBC_ARCHIVE_DIRNAME="@{pkg_name}-@{version}"
+CT_GLIBC_ARCHIVE_FORMATS=".tar.xz .tar.bz2 .tar.gz"
+CT_GLIBC_SIGNATURE_FORMAT="packed/.sig"
+CT_GLIBC_2_38_or_older=y
+CT_GLIBC_older_than_2_38=y
+CT_GLIBC_2_37_or_older=y
+CT_GLIBC_older_than_2_37=y
+CT_GLIBC_2_36_or_older=y
+CT_GLIBC_older_than_2_36=y
+CT_GLIBC_2_34_or_older=y
+CT_GLIBC_older_than_2_34=y
+CT_GLIBC_2_32_or_older=y
+CT_GLIBC_older_than_2_32=y
+CT_GLIBC_2_31_or_older=y
+CT_GLIBC_older_than_2_31=y
+CT_GLIBC_2_30_or_older=y
+CT_GLIBC_older_than_2_30=y
+CT_GLIBC_2_29_or_older=y
+CT_GLIBC_older_than_2_29=y
+CT_GLIBC_REQUIRE_older_than_2_29=y
+CT_GLIBC_2_28_or_later=y
+CT_GLIBC_2_28_or_older=y
+CT_GLIBC_later_than_2_27=y
+CT_GLIBC_2_27_or_later=y
+CT_GLIBC_later_than_2_26=y
+CT_GLIBC_2_26_or_later=y
+CT_GLIBC_later_than_2_25=y
+CT_GLIBC_2_25_or_later=y
+CT_GLIBC_later_than_2_24=y
+CT_GLIBC_2_24_or_later=y
+CT_GLIBC_later_than_2_23=y
+CT_GLIBC_2_23_or_later=y
+CT_GLIBC_later_than_2_20=y
+CT_GLIBC_2_20_or_later=y
+CT_GLIBC_later_than_2_17=y
+CT_GLIBC_2_17_or_later=y
+CT_GLIBC_later_than_2_14=y
+CT_GLIBC_2_14_or_later=y
+CT_GLIBC_DEP_KERNEL_HEADERS_VERSION=y
+CT_GLIBC_DEP_BINUTILS=y
+CT_GLIBC_DEP_GCC=y
+CT_GLIBC_DEP_PYTHON=y
+
+#
+# glibc >=2.29 requires python >3.4
+#
+CT_GLIBC_SPARC_ALLOW_V7=y
 CT_THREADS="nptl"
-# CT_CC_GLIBC_SHOW_LINARO is not set
-# CT_LIBC_GLIBC_V_2_22 is not set
-# CT_LIBC_GLIBC_V_2_21 is not set
-# CT_LIBC_GLIBC_V_2_20 is not set
-# CT_LIBC_GLIBC_V_2_19 is not set
-# CT_LIBC_GLIBC_V_2_18 is not set
-CT_LIBC_GLIBC_V_2_17=y
-# CT_LIBC_GLIBC_V_2_16_0 is not set
-# CT_LIBC_GLIBC_V_2_15 is not set
-# CT_LIBC_GLIBC_V_2_14_1 is not set
-# CT_LIBC_GLIBC_V_2_14 is not set
-# CT_LIBC_GLIBC_V_2_13 is not set
-# CT_LIBC_GLIBC_V_2_12_2 is not set
-# CT_LIBC_GLIBC_V_2_12_1 is not set
-# CT_LIBC_GLIBC_V_2_11_1 is not set
-# CT_LIBC_GLIBC_V_2_11 is not set
-# CT_LIBC_GLIBC_V_2_10_1 is not set
-# CT_LIBC_GLIBC_V_2_9 is not set
-# CT_LIBC_GLIBC_V_2_8 is not set
-CT_LIBC_GLIBC_2_17_or_later=y
-CT_LIBC_mingw_AVAILABLE=y
-CT_LIBC_musl_AVAILABLE=y
-CT_LIBC_newlib_AVAILABLE=y
-CT_LIBC_none_AVAILABLE=y
-CT_LIBC_uClibc_AVAILABLE=y
+CT_GLIBC_BUILD_SSP=y
+CT_GLIBC_HAS_LIBIDN_ADDON=y
+# CT_GLIBC_USE_LIBIDN_ADDON is not set
+CT_GLIBC_NO_SPARC_V8=y
+CT_GLIBC_HAS_OBSOLETE_RPC=y
+CT_GLIBC_EXTRA_CONFIG_ARRAY=""
+CT_GLIBC_CONFIGPARMS=""
+CT_GLIBC_ENABLE_DEBUG=y
+CT_GLIBC_EXTRA_CFLAGS="-Wno-missing-attributes -Wno-array-bounds -Wno-array-parameter -Wno-stringop-overflow -Wno-maybe-uninitialized"
+CT_GLIBC_ENABLE_OBSOLETE_RPC=y
+# CT_GLIBC_DISABLE_VERSIONING is not set
+CT_GLIBC_OLDEST_ABI=""
+CT_GLIBC_FORCE_UNWIND=y
+# CT_GLIBC_LOCALES is not set
+# CT_GLIBC_KERNEL_VERSION_NONE is not set
+CT_GLIBC_KERNEL_VERSION_AS_HEADERS=y
+# CT_GLIBC_KERNEL_VERSION_CHOSEN is not set
+CT_GLIBC_MIN_KERNEL="4.18.20"
+CT_GLIBC_SSP_DEFAULT=y
+# CT_GLIBC_SSP_NO is not set
+# CT_GLIBC_SSP_YES is not set
+# CT_GLIBC_SSP_ALL is not set
+# CT_GLIBC_SSP_STRONG is not set
+CT_GLIBC_ENABLE_COMMON_FLAG=y
+CT_ALL_LIBC_CHOICES="AVR_LIBC GLIBC MINGW_W64 MOXIEBOX MUSL NEWLIB NONE PICOLIBC UCLIBC_NG"
 CT_LIBC_SUPPORT_THREADS_ANY=y
 CT_LIBC_SUPPORT_THREADS_NATIVE=y
 
@@ -325,95 +508,107 @@ CT_LIBC_SUPPORT_THREADS_NATIVE=y
 # Common C library options
 #
 CT_THREADS_NATIVE=y
+# CT_CREATE_LDSO_CONF is not set
 CT_LIBC_XLDD=y
-
-#
-# glibc other options
-#
-# CT_LIBC_GLIBC_PORTS_EXTERNAL is not set
-CT_LIBC_glibc_familly=y
-CT_LIBC_GLIBC_EXTRA_CONFIG_ARRAY=""
-CT_LIBC_GLIBC_CONFIGPARMS=""
-CT_LIBC_GLIBC_EXTRA_CFLAGS=""
-CT_LIBC_EXTRA_CC_ARGS=""
-# CT_LIBC_DISABLE_VERSIONING is not set
-CT_LIBC_OLDEST_ABI=""
-CT_LIBC_GLIBC_FORCE_UNWIND=y
-# CT_LIBC_GLIBC_USE_PORTS is not set
-CT_LIBC_ADDONS_LIST=""
-# CT_LIBC_LOCALES is not set
-# CT_LIBC_GLIBC_KERNEL_VERSION_NONE is not set
-CT_LIBC_GLIBC_KERNEL_VERSION_AS_HEADERS=y
-# CT_LIBC_GLIBC_KERNEL_VERSION_CHOSEN is not set
-CT_LIBC_GLIBC_MIN_KERNEL="3.10.93"
+# end of C-library
 
 #
 # C compiler
 #
+CT_CC_CORE_NEEDED=y
+CT_CC_SUPPORT_CXX=y
+CT_CC_SUPPORT_FORTRAN=y
+CT_CC_SUPPORT_ADA=y
+CT_CC_SUPPORT_D=y
+CT_CC_SUPPORT_JIT=y
+CT_CC_SUPPORT_OBJC=y
+CT_CC_SUPPORT_OBJCXX=y
+CT_CC_SUPPORT_GOLANG=y
+CT_CC_GCC=y
 CT_CC="gcc"
-CT_CC_CORE_PASSES_NEEDED=y
-CT_CC_CORE_PASS_1_NEEDED=y
-CT_CC_CORE_PASS_2_NEEDED=y
-CT_CC_gcc=y
-# CT_CC_GCC_SHOW_LINARO is not set
-#CT_CC_GCC_V_5_2_0=y
-# CT_CC_GCC_V_4_9_3 is not set
-CT_CC_GCC_V_4_8_5=y
-# CT_CC_GCC_V_4_7_4 is not set
-# CT_CC_GCC_V_4_6_4 is not set
-# CT_CC_GCC_V_4_5_4 is not set
-# CT_CC_GCC_V_4_4_7 is not set
-# CT_CC_GCC_V_4_3_6 is not set
-# CT_CC_GCC_V_4_2_4 is not set
-CT_CC_GCC_4_2_or_later=y
-CT_CC_GCC_4_3_or_later=y
-CT_CC_GCC_4_4_or_later=y
-CT_CC_GCC_4_5_or_later=y
-CT_CC_GCC_4_6_or_later=y
-CT_CC_GCC_4_7_or_later=y
-CT_CC_GCC_4_8_or_later=y
-CT_CC_GCC_4_9_or_later=y
-CT_CC_GCC_5=y
-CT_CC_GCC_5_or_later=y
-CT_CC_GCC_HAS_GRAPHITE=y
-CT_CC_GCC_USE_GRAPHITE=y
-CT_CC_GCC_HAS_LTO=y
-CT_CC_GCC_USE_LTO=y
-CT_CC_GCC_HAS_PKGVERSION_BUGURL=y
-CT_CC_GCC_HAS_BUILD_ID=y
-CT_CC_GCC_HAS_LNK_HASH_STYLE=y
-CT_CC_GCC_USE_GMP_MPFR=y
-CT_CC_GCC_USE_MPC=y
-CT_CC_GCC_HAS_LIBQUADMATH=y
-CT_CC_GCC_HAS_LIBSANITIZER=y
-CT_CC_GCC_VERSION="4.8.5"
-#CT_CC_GCC_VERSION="5.2.0"
-# CT_CC_LANG_FORTRAN is not set
+CT_CC_CHOICE_KSYM="GCC"
+CT_CC_GCC_SHOW=y
+
+#
+# Options for gcc
+#
+CT_CC_GCC_PKG_KSYM="GCC"
+CT_GCC_DIR_NAME="gcc"
+CT_GCC_USE_GNU=y
+# CT_GCC_USE_ORACLE is not set
+CT_GCC_USE="GCC"
+CT_GCC_PKG_NAME="gcc"
+CT_GCC_SRC_RELEASE=y
+# CT_GCC_SRC_DEVEL is not set
+CT_GCC_PATCH_ORDER="global"
+CT_GCC_V_13=y
+# CT_GCC_V_12 is not set
+# CT_GCC_V_11 is not set
+# CT_GCC_V_10 is not set
+# CT_GCC_V_9 is not set
+# CT_GCC_V_8 is not set
+# CT_GCC_V_7 is not set
+# CT_GCC_V_6 is not set
+# CT_GCC_V_5 is not set
+# CT_GCC_V_4_9 is not set
+CT_GCC_VERSION="13.2.0"
+CT_GCC_MIRRORS="$(CT_Mirrors GNU gcc/gcc-${CT_GCC_VERSION}) $(CT_Mirrors sourceware gcc/releases/gcc-${CT_GCC_VERSION})"
+CT_GCC_ARCHIVE_FILENAME="@{pkg_name}-@{version}"
+CT_GCC_ARCHIVE_DIRNAME="@{pkg_name}-@{version}"
+CT_GCC_ARCHIVE_FORMATS=".tar.xz .tar.gz"
+CT_GCC_SIGNATURE_FORMAT=""
+CT_GCC_later_than_13=y
+CT_GCC_13_or_later=y
+CT_GCC_later_than_12=y
+CT_GCC_12_or_later=y
+CT_GCC_later_than_11=y
+CT_GCC_11_or_later=y
+CT_GCC_later_than_10=y
+CT_GCC_10_or_later=y
+CT_GCC_later_than_9=y
+CT_GCC_9_or_later=y
+CT_GCC_later_than_8=y
+CT_GCC_8_or_later=y
+CT_GCC_later_than_7=y
+CT_GCC_7_or_later=y
+CT_GCC_later_than_6=y
+CT_GCC_6_or_later=y
+CT_GCC_later_than_5=y
+CT_GCC_5_or_later=y
+CT_GCC_later_than_4_9=y
+CT_GCC_4_9_or_later=y
+CT_GCC_REQUIRE_4_9_or_later=y
+CT_CC_GCC_HAS_LIBMPX=y
 CT_CC_GCC_ENABLE_CXX_FLAGS=""
 CT_CC_GCC_CORE_EXTRA_CONFIG_ARRAY=""
 CT_CC_GCC_EXTRA_CONFIG_ARRAY=""
-CT_CC_GCC_EXTRA_ENV_ARRAY=""
 CT_CC_GCC_STATIC_LIBSTDCXX=y
-# CT_CC_GCC_SYSTEM_ZLIB is not set
+CT_CC_GCC_SYSTEM_ZLIB=y
+CT_CC_GCC_CONFIG_TLS=m
 
 #
 # Optimisation features
 #
+CT_CC_GCC_USE_GRAPHITE=y
 
 #
 # Settings for libraries running on target
 #
+# CT_CC_GCC_ENABLE_DEFAULT_PIE is not set
 CT_CC_GCC_ENABLE_TARGET_OPTSPACE=y
 # CT_CC_GCC_LIBMUDFLAP is not set
 # CT_CC_GCC_LIBGOMP is not set
 # CT_CC_GCC_LIBSSP is not set
 # CT_CC_GCC_LIBQUADMATH is not set
 # CT_CC_GCC_LIBSANITIZER is not set
+CT_CC_GCC_LIBMPX=y
+CT_CC_GCC_LIBSTDCXX_VERBOSE=m
 
 #
 # Misc. obscure options.
 #
 CT_CC_CXA_ATEXIT=y
+CT_CC_GCC_TM_CLONE_REGISTRY=m
 CT_CC_GCC_DISABLE_PCH=y
 CT_CC_GCC_SJLJ_EXCEPTIONS=m
 CT_CC_GCC_LDBL_128=m
@@ -423,105 +618,243 @@ CT_CC_GCC_LNK_HASH_STYLE_DEFAULT=y
 # CT_CC_GCC_LNK_HASH_STYLE_GNU is not set
 # CT_CC_GCC_LNK_HASH_STYLE_BOTH is not set
 CT_CC_GCC_LNK_HASH_STYLE=""
-CT_CC_GCC_DEC_FLOAT_AUTO=y
-# CT_CC_GCC_DEC_FLOAT_BID is not set
-# CT_CC_GCC_DEC_FLOAT_DPD is not set
+CT_CC_GCC_DEC_FLOATS_AUTO=y
+# CT_CC_GCC_DEC_FLOATS_BID is not set
+# CT_CC_GCC_DEC_FLOATS_DPD is not set
 # CT_CC_GCC_DEC_FLOATS_NO is not set
-CT_CC_SUPPORT_CXX=y
-CT_CC_SUPPORT_FORTRAN=y
-CT_CC_SUPPORT_JAVA=y
-CT_CC_SUPPORT_ADA=y
-CT_CC_SUPPORT_OBJC=y
-CT_CC_SUPPORT_OBJCXX=y
-CT_CC_SUPPORT_GOLANG=y
+CT_CC_GCC_DEC_FLOATS=""
+CT_ALL_CC_CHOICES="GCC"
 
 #
 # Additional supported languages:
 #
 CT_CC_LANG_CXX=y
-# CT_CC_LANG_JAVA is not set
+# CT_CC_LANG_FORTRAN is not set
+# end of C compiler
 
 #
 # Debug facilities
 #
-# CT_DEBUG_dmalloc is not set
-# CT_DEBUG_duma is not set
-# CT_DEBUG_gdb is not set
-# CT_DEBUG_ltrace is not set
-# CT_DEBUG_strace is not set
+# CT_DEBUG_DUMA is not set
+# CT_DEBUG_GDB is not set
+# CT_DEBUG_LTRACE is not set
+# CT_DEBUG_STRACE is not set
+CT_ALL_DEBUG_CHOICES="DUMA GDB LTRACE STRACE"
+# end of Debug facilities
 
 #
 # Companion libraries
 #
-CT_COMPLIBS_NEEDED=y
+# CT_COMPLIBS_CHECK is not set
+# CT_COMP_LIBS_CLOOG is not set
+# CT_COMP_LIBS_EXPAT is not set
+CT_COMP_LIBS_GETTEXT=y
+CT_COMP_LIBS_GETTEXT_PKG_KSYM="GETTEXT"
+CT_GETTEXT_DIR_NAME="gettext"
+CT_GETTEXT_PKG_NAME="gettext"
+CT_GETTEXT_SRC_RELEASE=y
+# CT_GETTEXT_SRC_DEVEL is not set
+CT_GETTEXT_PATCH_ORDER="global"
+CT_GETTEXT_V_0_21=y
+# CT_GETTEXT_V_0_20_1 is not set
+# CT_GETTEXT_V_0_19_8_1 is not set
+CT_GETTEXT_VERSION="0.21"
+CT_GETTEXT_MIRRORS="$(CT_Mirrors GNU gettext)"
+CT_GETTEXT_ARCHIVE_FILENAME="@{pkg_name}-@{version}"
+CT_GETTEXT_ARCHIVE_DIRNAME="@{pkg_name}-@{version}"
+CT_GETTEXT_ARCHIVE_FORMATS=".tar.xz .tar.gz"
+CT_GETTEXT_SIGNATURE_FORMAT="packed/.sig"
+CT_GETTEXT_0_21_or_later=y
+CT_GETTEXT_0_21_or_older=y
+CT_GETTEXT_INCOMPATIBLE_WITH_UCLIBC_NG=y
+
+#
+# This version of gettext is not compatible with uClibc-NG. Select
+#
+
+#
+# a different version if uClibc-NG is used on the target or (in a
+#
+
+#
+# Canadian cross build) on the host.
+#
+CT_COMP_LIBS_GMP=y
+CT_COMP_LIBS_GMP_PKG_KSYM="GMP"
+CT_GMP_DIR_NAME="gmp"
+CT_GMP_PKG_NAME="gmp"
+CT_GMP_SRC_RELEASE=y
+# CT_GMP_SRC_DEVEL is not set
+CT_GMP_PATCH_ORDER="global"
+CT_GMP_V_6_2=y
+# CT_GMP_V_6_1 is not set
+CT_GMP_VERSION="6.2.1"
+CT_GMP_MIRRORS="https://gmplib.org/download/gmp https://gmplib.org/download/gmp/archive $(CT_Mirrors GNU gmp)"
+CT_GMP_ARCHIVE_FILENAME="@{pkg_name}-@{version}"
+CT_GMP_ARCHIVE_DIRNAME="@{pkg_name}-@{version}"
+CT_GMP_ARCHIVE_FORMATS=".tar.xz .tar.lz .tar.bz2"
+CT_GMP_SIGNATURE_FORMAT="packed/.sig"
+CT_COMP_LIBS_ISL=y
+CT_COMP_LIBS_ISL_PKG_KSYM="ISL"
+CT_ISL_DIR_NAME="isl"
+CT_ISL_PKG_NAME="isl"
+CT_ISL_SRC_RELEASE=y
+# CT_ISL_SRC_DEVEL is not set
+CT_ISL_PATCH_ORDER="global"
+CT_ISL_V_0_26=y
+# CT_ISL_V_0_25 is not set
+# CT_ISL_V_0_24 is not set
+# CT_ISL_V_0_23 is not set
+# CT_ISL_V_0_22 is not set
+# CT_ISL_V_0_21 is not set
+# CT_ISL_V_0_20 is not set
+# CT_ISL_V_0_19 is not set
+# CT_ISL_V_0_18 is not set
+# CT_ISL_V_0_17 is not set
+# CT_ISL_V_0_16 is not set
+# CT_ISL_V_0_15 is not set
+CT_ISL_VERSION="0.26"
+CT_ISL_MIRRORS="https://libisl.sourceforge.io"
+CT_ISL_ARCHIVE_FILENAME="@{pkg_name}-@{version}"
+CT_ISL_ARCHIVE_DIRNAME="@{pkg_name}-@{version}"
+CT_ISL_ARCHIVE_FORMATS=".tar.xz .tar.bz2 .tar.gz"
+CT_ISL_SIGNATURE_FORMAT=""
+CT_ISL_later_than_0_18=y
+CT_ISL_0_18_or_later=y
+CT_ISL_later_than_0_15=y
+CT_ISL_0_15_or_later=y
+# CT_COMP_LIBS_LIBELF is not set
+CT_COMP_LIBS_LIBICONV=y
+CT_COMP_LIBS_LIBICONV_PKG_KSYM="LIBICONV"
+CT_LIBICONV_DIR_NAME="libiconv"
+CT_LIBICONV_PKG_NAME="libiconv"
+CT_LIBICONV_SRC_RELEASE=y
+# CT_LIBICONV_SRC_DEVEL is not set
+CT_LIBICONV_PATCH_ORDER="global"
+CT_LIBICONV_V_1_16=y
+# CT_LIBICONV_V_1_15 is not set
+CT_LIBICONV_VERSION="1.16"
+CT_LIBICONV_MIRRORS="$(CT_Mirrors GNU libiconv)"
+CT_LIBICONV_ARCHIVE_FILENAME="@{pkg_name}-@{version}"
+CT_LIBICONV_ARCHIVE_DIRNAME="@{pkg_name}-@{version}"
+CT_LIBICONV_ARCHIVE_FORMATS=".tar.gz"
+CT_LIBICONV_SIGNATURE_FORMAT="packed/.sig"
+CT_COMP_LIBS_MPC=y
+CT_COMP_LIBS_MPC_PKG_KSYM="MPC"
+CT_MPC_DIR_NAME="mpc"
+CT_MPC_PKG_NAME="mpc"
+CT_MPC_SRC_RELEASE=y
+# CT_MPC_SRC_DEVEL is not set
+CT_MPC_PATCH_ORDER="global"
+CT_MPC_V_1_2=y
+CT_MPC_VERSION="1.2.1"
+CT_MPC_MIRRORS="https://www.multiprecision.org/downloads $(CT_Mirrors GNU mpc)"
+CT_MPC_ARCHIVE_FILENAME="@{pkg_name}-@{version}"
+CT_MPC_ARCHIVE_DIRNAME="@{pkg_name}-@{version}"
+CT_MPC_ARCHIVE_FORMATS=".tar.gz"
+CT_MPC_SIGNATURE_FORMAT="packed/.sig"
+CT_MPC_later_than_1_1_0=y
+CT_MPC_1_1_0_or_later=y
+CT_COMP_LIBS_MPFR=y
+CT_COMP_LIBS_MPFR_PKG_KSYM="MPFR"
+CT_MPFR_DIR_NAME="mpfr"
+CT_MPFR_PKG_NAME="mpfr"
+CT_MPFR_SRC_RELEASE=y
+# CT_MPFR_SRC_DEVEL is not set
+CT_MPFR_PATCH_ORDER="global"
+CT_MPFR_V_4_2=y
+CT_MPFR_VERSION="4.2.1"
+CT_MPFR_MIRRORS="https://www.mpfr.org/mpfr-${CT_MPFR_VERSION} $(CT_Mirrors GNU mpfr)"
+CT_MPFR_ARCHIVE_FILENAME="@{pkg_name}-@{version}"
+CT_MPFR_ARCHIVE_DIRNAME="@{pkg_name}-@{version}"
+CT_MPFR_ARCHIVE_FORMATS=".tar.xz .tar.bz2 .tar.gz .zip"
+CT_MPFR_SIGNATURE_FORMAT="packed/.asc"
+CT_MPFR_later_than_4_0_0=y
+CT_MPFR_4_0_0_or_later=y
+CT_COMP_LIBS_NCURSES=y
+CT_COMP_LIBS_NCURSES_PKG_KSYM="NCURSES"
+CT_NCURSES_DIR_NAME="ncurses"
+CT_NCURSES_PKG_NAME="ncurses"
+CT_NCURSES_SRC_RELEASE=y
+# CT_NCURSES_SRC_DEVEL is not set
+CT_NCURSES_PATCH_ORDER="global"
+CT_NCURSES_V_6_4=y
+# CT_NCURSES_V_6_2 is not set
+# CT_NCURSES_V_6_1 is not set
+# CT_NCURSES_V_6_0 is not set
+CT_NCURSES_VERSION="6.4"
+CT_NCURSES_MIRRORS="https://invisible-mirror.net/archives/ncurses $(CT_Mirrors GNU ncurses)"
+CT_NCURSES_ARCHIVE_FILENAME="@{pkg_name}-@{version}"
+CT_NCURSES_ARCHIVE_DIRNAME="@{pkg_name}-@{version}"
+CT_NCURSES_ARCHIVE_FORMATS=".tar.gz"
+CT_NCURSES_SIGNATURE_FORMAT="packed/.sig"
+CT_NCURSES_NEW_ABI=y
+CT_NCURSES_HOST_CONFIG_ARGS=""
+CT_NCURSES_HOST_DISABLE_DB=y
+CT_NCURSES_HOST_FALLBACKS="linux,xterm,xterm-color,xterm-256color,vt100"
+CT_NCURSES_TARGET_CONFIG_ARGS=""
+# CT_NCURSES_TARGET_DISABLE_DB is not set
+CT_NCURSES_TARGET_FALLBACKS=""
+CT_COMP_LIBS_ZLIB=y
+CT_COMP_LIBS_ZLIB_PKG_KSYM="ZLIB"
+CT_ZLIB_DIR_NAME="zlib"
+CT_ZLIB_PKG_NAME="zlib"
+CT_ZLIB_SRC_RELEASE=y
+# CT_ZLIB_SRC_DEVEL is not set
+CT_ZLIB_PATCH_ORDER="global"
+CT_ZLIB_V_1_2_13=y
+CT_ZLIB_VERSION="1.2.13"
+CT_ZLIB_MIRRORS="[REDACTED_PROFILE_URL]/zlib/releases/download/v${CT_ZLIB_VERSION} https://www.zlib.net/"
+CT_ZLIB_ARCHIVE_FILENAME="@{pkg_name}-@{version}"
+CT_ZLIB_ARCHIVE_DIRNAME="@{pkg_name}-@{version}"
+CT_ZLIB_ARCHIVE_FORMATS=".tar.xz .tar.gz"
+CT_ZLIB_SIGNATURE_FORMAT="packed/.asc"
+CT_COMP_LIBS_ZSTD=y
+CT_COMP_LIBS_ZSTD_PKG_KSYM="ZSTD"
+CT_ZSTD_DIR_NAME="zstd"
+CT_ZSTD_PKG_NAME="zstd"
+CT_ZSTD_SRC_RELEASE=y
+# CT_ZSTD_SRC_DEVEL is not set
+CT_ZSTD_PATCH_ORDER="global"
+CT_ZSTD_V_1_5_5=y
+# CT_ZSTD_V_1_5_2 is not set
+CT_ZSTD_VERSION="1.5.5"
+CT_ZSTD_MIRRORS="[REDACTED_PROFILE_URL]/zstd/releases/download/v${CT_ZSTD_VERSION} https://www.zstd.net/"
+CT_ZSTD_ARCHIVE_FILENAME="@{pkg_name}-@{version}"
+CT_ZSTD_ARCHIVE_DIRNAME="@{pkg_name}-@{version}"
+CT_ZSTD_ARCHIVE_FORMATS=".tar.gz"
+CT_ZSTD_SIGNATURE_FORMAT="packed/.sig"
+CT_ALL_COMP_LIBS_CHOICES="CLOOG EXPAT GETTEXT GMP GNUPRUMCU ISL LIBELF LIBICONV MPC MPFR NCURSES NEWLIB_NANO PICOLIBC ZLIB ZSTD"
 CT_LIBICONV_NEEDED=y
 CT_GETTEXT_NEEDED=y
 CT_GMP_NEEDED=y
 CT_MPFR_NEEDED=y
 CT_ISL_NEEDED=y
-CT_CLOOG_NEEDED=y
 CT_MPC_NEEDED=y
-CT_COMPLIBS=y
+CT_NCURSES_NEEDED=y
+CT_ZLIB_NEEDED=y
+CT_ZSTD_NEEDED=y
 CT_LIBICONV=y
 CT_GETTEXT=y
 CT_GMP=y
 CT_MPFR=y
 CT_ISL=y
-CT_CLOOG=y
 CT_MPC=y
-CT_LIBICONV_V_1_14=y
-CT_LIBICONV_VERSION="1.14"
-CT_GETTEXT_V_0_19_6=y
-CT_GETTEXT_VERSION="0.19.6"
-CT_GMP_V_6_0_0=y
-# CT_GMP_V_5_1_3 is not set
-# CT_GMP_V_5_1_1 is not set
-# CT_GMP_V_5_0_2 is not set
-# CT_GMP_V_5_0_1 is not set
-# CT_GMP_V_4_3_2 is not set
-# CT_GMP_V_4_3_1 is not set
-# CT_GMP_V_4_3_0 is not set
-CT_GMP_5_0_2_or_later=y
-CT_GMP_VERSION="6.0.0a"
-CT_MPFR_V_3_1_3=y
-# CT_MPFR_V_3_1_2 is not set
-# CT_MPFR_V_3_1_0 is not set
-# CT_MPFR_V_3_0_1 is not set
-# CT_MPFR_V_3_0_0 is not set
-# CT_MPFR_V_2_4_2 is not set
-# CT_MPFR_V_2_4_1 is not set
-# CT_MPFR_V_2_4_0 is not set
-CT_MPFR_VERSION="3.1.3"
-CT_ISL_V_0_14=y
-CT_ISL_V_0_14_or_later=y
-CT_ISL_V_0_12_or_later=y
-CT_ISL_VERSION="0.14"
-CT_CLOOG_V_0_18_4=y
-# CT_CLOOG_V_0_18_1 is not set
-# CT_CLOOG_V_0_18_0 is not set
-CT_CLOOG_VERSION="0.18.4"
-CT_CLOOG_0_18_4_or_later=y
-CT_CLOOG_0_18_or_later=y
-CT_MPC_V_1_0_3=y
-# CT_MPC_V_1_0_2 is not set
-# CT_MPC_V_1_0_1 is not set
-# CT_MPC_V_1_0 is not set
-# CT_MPC_V_0_9 is not set
-# CT_MPC_V_0_8_2 is not set
-# CT_MPC_V_0_8_1 is not set
-# CT_MPC_V_0_7 is not set
-CT_MPC_VERSION="1.0.3"
-
-#
-# Companion libraries common options
-#
-# CT_COMPLIBS_CHECK is not set
+CT_NCURSES=y
+CT_ZLIB=y
+CT_ZSTD=y
+# end of Companion libraries
 
 #
 # Companion tools
 #
-
-#
-# READ HELP before you say 'Y' below !!!
-#
-# CT_COMP_TOOLS is not set
+# CT_COMP_TOOLS_FOR_HOST is not set
+# CT_COMP_TOOLS_AUTOCONF is not set
+# CT_COMP_TOOLS_AUTOMAKE is not set
+# CT_COMP_TOOLS_BISON is not set
+# CT_COMP_TOOLS_DTC is not set
+# CT_COMP_TOOLS_LIBTOOL is not set
+# CT_COMP_TOOLS_M4 is not set
+# CT_COMP_TOOLS_MAKE is not set
+CT_ALL_COMP_TOOLS_CHOICES="AUTOCONF AUTOMAKE BISON DTC LIBTOOL M4 MAKE"
+# end of Companion tools
diff --git a/Engine/Config/Linux/Linux_SDK.json b/Engine/Config/Linux/Linux_SDK.json
--- a/Engine/Config/Linux/Linux_SDK.json
+++ b/Engine/Config/Linux/Linux_SDK.json
@@ -1,7 +1,7 @@
 {
-"MainVersion" : "v22_clang-16.0.6-centos7",
-"MinVersion" : "v22_clang-16.0.6-centos7",
-"MaxVersion" : "v22_clang-16.0.6-centos7",
+"MainVersion" : "v23_clang-18.1.0-rockylinux8",
+"MinVersion" : "v23_clang-18.1.0-rockylinux8",
+"MaxVersion" : "v23_clang-18.1.0-rockylinux8",
 
 "AutoSDKPlatform" : "Linux_x64",
 }
diff --git a/Engine/Programs/UnrealGameSync/UnrealGameSync.ini b/Engine/Programs/UnrealGameSync/UnrealGameSync.ini
--- a/Engine/Programs/UnrealGameSync/UnrealGameSync.ini
+++ b/Engine/Programs/UnrealGameSync/UnrealGameSync.ini
@@ -55,6 +55,6 @@ SdkInstallerDir=
 +SdkInfo=(Category="Windows", Description="Visual Studio 2022 v18.4", Browse="https://visualstudio.microsoft.com/downloads/")
 +SdkInfo=(Category="Windows", Description="Windows SDK 10.0.18362.0", Browse="https://developer.microsoft.com/en-us/windows/downloads/windows-sdk/")
 +SdkInfo=(Category="Windows", Description="LLVM Clang 16.0.6", Browse="https://releases.llvm.org/")
-+SdkInfo=(Category="Linux", Description="v22_clang-16.0.6-centos7", Install="$(SdkInstallerDir)\\Linux\\HostWin64\\v22_clang-16.0.6-centos7.exe", Browse="$(SdkInstallerDir)\\Linux\\HostWin64")
++SdkInfo=(Category="Linux", Description="v23_clang-18.1.0-rockylinux8", Install="$(SdkInstallerDir)\\Linux\\HostWin64\\v23_clang-18.1.0-rockylinux8.exe", Browse="$(SdkInstallerDir)\\Linux\\HostWin64")
 +SdkInfo=(Category="Mac", Description="macOS 10.15.7")
 +SdkInfo=(Category="Mac", Description="Xcode 12.4")
diff --git a/Engine/Source/Programs/UnrealBuildTool/Platform/Linux/LinuxToolChain.cs b/Engine/Source/Programs/UnrealBuildTool/Platform/Linux/LinuxToolChain.cs
--- a/Engine/Source/Programs/UnrealBuildTool/Platform/Linux/LinuxToolChain.cs
+++ b/Engine/Source/Programs/UnrealBuildTool/Platform/Linux/LinuxToolChain.cs
@@ -55,7 +55,7 @@ namespace UnrealBuildTool
 			: this(UnrealTargetPlatform.Linux, InArchitecture, InSDK, InOptions, InLogger)
 		{
 			// prevent unknown clangs since the build is likely to fail on too old or too new compilers
-			if ((CompilerVersionLessThan(16, 0, 0) || CompilerVersionGreaterOrEqual(17, 0, 0)) && !Options.HasFlag(ClangToolChainOptions.UseAutoRTFMCompiler))
+			if ((CompilerVersionLessThan(16, 0, 0) || CompilerVersionGreaterOrEqual(19, 0, 0)) && !Options.HasFlag(ClangToolChainOptions.UseAutoRTFMCompiler))
 			{
 				throw new BuildException(
 					String.Format("This version of the Unreal Engine can only be compiled with clang 16.x. clang {0} may not build it - please use a different version.",
