Include all of the AutomationTool source and project when testing for dirty dependencies. Rely on findstr to filter files to include rather than dir /s with file extensions.

#rnx

diff --git a/Engine/Build/BatchFiles/BuildUAT.bat b/Engine/Build/BatchFiles/BuildUAT.bat
--- a/Engine/Build/BatchFiles/BuildUAT.bat
+++ b/Engine/Build/BatchFiles/BuildUAT.bat
@@ -9,7 +9,7 @@ rem ## if you copy it to a different location and run it.
 setlocal
 
 rem ## Make sure we use the Windows "find" utility and not a Unix-like tool found in PATH
-set FIND=%SYSTEMROOT%\System32\find.exe
+set FIND=%SYSTEMROOT%\System32\findstr.exe /i /r
 
 rem ## First, make sure the batch file exists in the folder we expect it to.  This is necessary in order to
 rem ## verify that our relative path to the /Engine/Source directory is correct
@@ -24,69 +24,52 @@ rem find ".cs" files to only lines that match those names - excludes lines that
 md ..\Intermediate\Build >nul 2>nul
 
 dir /s ^
- Programs\Shared\EpicGames.Build\*.cs ^
- Programs\Shared\EpicGames.Build\*.csproj ^
- Programs\Shared\EpicGames.Core\*.cs ^
- Programs\Shared\EpicGames.Core\*.csproj ^
- Programs\Shared\EpicGames.Horde\*.cs ^
- Programs\Shared\EpicGames.Horde\*.csproj ^
- Programs\Shared\EpicGames.IoHash\*.cs ^
- Programs\Shared\EpicGames.IoHash\*.csproj ^
- Programs\Shared\EpicGames.MsBuild\*.cs ^
- Programs\Shared\EpicGames.MsBuild\*.csproj ^
- Programs\Shared\EpicGames.OIDC\*.cs ^
- Programs\Shared\EpicGames.OIDC\*.csproj ^
- Programs\Shared\EpicGames.Serialization\*.cs ^
- Programs\Shared\EpicGames.Serialization\*.csproj ^
- Programs\Shared\EpicGames.UBA\*.cs ^
- Programs\Shared\EpicGames.UBA\*.csproj ^
- Programs\Shared\EpicGames.UHT\*.cs ^
- Programs\Shared\EpicGames.UHT\*.csproj ^
- Programs\UnrealBuildTool\*.cs ^
- Programs\UnrealBuildTool\*.csproj ^
- | %FIND% ".cs" > ..\Intermediate\Build\AutomationToolFiles.txt
-
-if exist ..\Binaries\Win64\UnrealBuildAccelerator (
-	dir /s ^
-	 ..\Binaries\Win64\UnrealBuildAccelerator\*.dll ^
-	 | %FIND% ".dll" >> ..\Intermediate\Build\AutomationToolFiles.txt
-	dir /s ^
-	 ..\Binaries\Win64\UnrealBuildAccelerator\*.exe ^
-	 | %FIND% ".exe" >> ..\Intermediate\Build\AutomationToolFiles.txt
-) 2>nul
+ Programs\Shared\EpicGames.Build ^
+ Programs\Shared\EpicGames.Core ^
+ Programs\Shared\EpicGames.Horde ^
+ Programs\Shared\EpicGames.IoHash ^
+ Programs\Shared\EpicGames.MsBuild ^
+ Programs\Shared\EpicGames.OIDC ^
+ Programs\Shared\EpicGames.Serialization ^
+ Programs\Shared\EpicGames.UBA ^
+ Programs\Shared\EpicGames.UHT ^
+ Programs\UnrealBuildTool ^
+ Programs\AutomationTool ^
+ 2>nul ^
+ | %FIND% /c:"\.cs$" /c:"\.csproj$" > ..\Intermediate\Build\AutomationToolFiles.txt
+
+dir /s ^
+ ..\Binaries\Win64\UnrealBuildAccelerator ^
+ 2>nul ^
+ | %FIND% /c:"\.dll$" /c:"\.exe$" >> ..\Intermediate\Build\AutomationToolFiles.txt
 
 if not exist ..\Platforms goto NoPlatforms
 for /d %%D in (..\Platforms\*) do (
-	for %%F in (AutomationTool Shared UnrealBuildTool) do (
-		if exist %%D\Source\Programs\%%F (
-			dir /s ^
-			%%D\Source\Programs\%%F\*.cs ^
-			%%D\Source\Programs\%%F\*.csproj ^
-			| %FIND% ".cs" >> ..\Intermediate\Build\AutomationToolFiles.txt
-		) 2>nul
-	)
+	dir /s ^
+	 %%D\Source\Programs\AutomationTool ^
+	 %%D\Source\Programs\Shared ^
+	 %%D\Source\Programs\UnrealBuildTool ^
+	 2>nul ^
+	 | %FIND% /c:"\.cs$" /c:"\.csproj$" >> ..\Intermediate\Build\AutomationToolFiles.txt
 )
 :NoPlatforms
 
 if not exist ..\Restricted goto NoRestricted
 for /d %%D in (..\Restricted\*) do (
-	for %%F in (AutomationTool Shared UnrealBuildTool) do (
-		if exist %%D\Source\Programs\%%F (
-			dir /s ^
-			%%D\Source\Programs\%%F\*.cs ^
-			%%D\Source\Programs\%%F\*.csproj ^
-			| %FIND% ".cs" >> ..\Intermediate\Build\AutomationToolFiles.txt
-		) 2>nul
-	)
+	dir /s ^
+	 %%D\Source\Programs\AutomationTool ^
+	 %%D\Source\Programs\Shared ^
+	 %%D\Source\Programs\UnrealBuildTool ^
+	 2>nul ^
+	 | %FIND% /c:"\.cs$" /c:"\.csproj$" >> ..\Intermediate\Build\AutomationToolFiles.txt
 )
 :NoRestricted
 
 rem note: no /s
 dir ^
  Programs\Shared\MetaData.cs ^
- Programs\AutomationTool\*.cs ^
- Programs\AutomationTool\*.csproj ^
- | %FIND% ".cs" >>..\Intermediate\Build\AutomationToolFiles.txt
+ 2>nul ^
+ | %FIND% /c:"\.cs$" >>..\Intermediate\Build\AutomationToolFiles.txt
 
 set MSBUILD_LOGLEVEL=%1
 if not defined %MSBUILD_LOGLEVEL set MSBUILD_LOGLEVEL=quiet
diff --git a/Engine/Build/BatchFiles/BuildUBT.bat b/Engine/Build/BatchFiles/BuildUBT.bat
--- a/Engine/Build/BatchFiles/BuildUBT.bat
+++ b/Engine/Build/BatchFiles/BuildUBT.bat
@@ -9,7 +9,7 @@ rem ## if you copy it to a different location and run it.
 setlocal
 
 rem ## Make sure we use the Windows "find" utility and not a Unix-like tool found in PATH
-set FIND=%SYSTEMROOT%\System32\find.exe
+set FIND=%SYSTEMROOT%\System32\findstr.exe /i /r
 
 rem ## First, make sure the batch file exists in the folder we expect it to.  This is necessary in order to
 rem ## verify that our relative path to the /Engine/Source directory is correct
@@ -24,67 +24,49 @@ rem find ".cs" files to only lines that match those names - excludes lines that
 md ..\Intermediate\Build >nul 2>nul
 
 dir /s ^
- Programs\Shared\EpicGames.Build\*.cs ^
- Programs\Shared\EpicGames.Build\*.csproj ^
- Programs\Shared\EpicGames.Core\*.cs ^
- Programs\Shared\EpicGames.Core\*.csproj ^
- Programs\Shared\EpicGames.Horde\*.cs ^
- Programs\Shared\EpicGames.Horde\*.csproj ^
- Programs\Shared\EpicGames.IoHash\*.cs ^
- Programs\Shared\EpicGames.IoHash\*.csproj ^
- Programs\Shared\EpicGames.MsBuild\*.cs ^
- Programs\Shared\EpicGames.MsBuild\*.csproj ^
- Programs\Shared\EpicGames.OIDC\*.cs ^
- Programs\Shared\EpicGames.OIDC\*.csproj ^
- Programs\Shared\EpicGames.Serialization\*.cs ^
- Programs\Shared\EpicGames.Serialization\*.csproj ^
- Programs\Shared\EpicGames.UBA\*.cs ^
- Programs\Shared\EpicGames.UBA\*.csproj ^
- Programs\Shared\EpicGames.UHT\*.cs ^
- Programs\Shared\EpicGames.UHT\*.csproj ^
- Programs\UnrealBuildTool\*.cs ^
- Programs\UnrealBuildTool\*.csproj ^
- | %FIND% ".cs" > ..\Intermediate\Build\UnrealBuildToolFiles.txt
-
-if exist ..\Binaries\Win64\UnrealBuildAccelerator (
-	dir /s ^
-	 ..\Binaries\Win64\UnrealBuildAccelerator\*.dll ^
-	 | %FIND% ".dll" >> ..\Intermediate\Build\UnrealBuildToolFiles.txt
-	dir /s ^
-	 ..\Binaries\Win64\UnrealBuildAccelerator\*.exe ^
-	 | %FIND% ".exe" >> ..\Intermediate\Build\UnrealBuildToolFiles.txt
-) 2>nul
+ Programs\Shared\EpicGames.Build ^
+ Programs\Shared\EpicGames.Core ^
+ Programs\Shared\EpicGames.Horde ^
+ Programs\Shared\EpicGames.IoHash ^
+ Programs\Shared\EpicGames.MsBuild ^
+ Programs\Shared\EpicGames.OIDC ^
+ Programs\Shared\EpicGames.Serialization ^
+ Programs\Shared\EpicGames.UBA ^
+ Programs\Shared\EpicGames.UHT ^
+ Programs\UnrealBuildTool ^
+ 2>nul ^
+ | %FIND% /c:"\.cs$" /c:"\.csproj$" > ..\Intermediate\Build\UnrealBuildToolFiles.txt
+
+dir /s ^
+ ..\Binaries\Win64\UnrealBuildAccelerator ^
+ 2>nul ^
+ | %FIND% /c:"\.dll$" /c:"\.exe$" >> ..\Intermediate\Build\UnrealBuildToolFiles.txt
 
 if not exist ..\Platforms goto NoPlatforms
 for /d %%D in (..\Platforms\*) do (
-	for %%F in (Shared UnrealBuildTool) do (
-		if exist %%D\Source\Programs\%%F (
-			dir /s ^
-			%%D\Source\Programs\%%F\*.cs ^
-			%%D\Source\Programs\%%F\*.csproj ^
-			| %FIND% ".cs" >> ..\Intermediate\Build\UnrealBuildToolFiles.txt
-		) 2>nul
-	)
+	dir /s ^
+	 %%D\Source\Programs\Shared ^
+	 %%D\Source\Programs\UnrealBuildTool ^
+	 2>nul ^
+	 | %FIND% /c:"\.cs$" /c:"\.csproj$" >> ..\Intermediate\Build\UnrealBuildToolFiles.txt
 )
 :NoPlatforms
 
 if not exist ..\Restricted goto NoRestricted
 for /d %%D in (..\Restricted\*) do (
-	for %%F in (Shared UnrealBuildTool) do (
-		if exist %%D\Source\Programs\%%F (
-			dir /s ^
-			%%D\Source\Programs\%%F\*.cs ^
-			%%D\Source\Programs\%%F\*.csproj ^
-			| %FIND% ".cs" >> ..\Intermediate\Build\UnrealBuildToolFiles.txt
-		) 2>nul
-	)
+	dir /s ^
+	 %%D\Source\Programs\Shared ^
+	 %%D\Source\Programs\UnrealBuildTool ^
+	 2>nul ^
+	 | %FIND% /c:"\.cs$" /c:"\.csproj$" >> ..\Intermediate\Build\UnrealBuildToolFiles.txt
 )
 :NoRestricted
 
 rem note: no /s
 dir ^
  Programs\Shared\MetaData.cs ^
- | %FIND% ".cs" >>..\Intermediate\Build\UnrealBuildToolFiles.txt
+ 2>nul ^
+ | %FIND% "\.cs$" >>..\Intermediate\Build\UnrealBuildToolFiles.txt
 
 set MSBUILD_LOGLEVEL=%1
 if not defined %MSBUILD_LOGLEVEL set MSBUILD_LOGLEVEL=quiet
