Fix new linux toolchain config to not compress debug sections since they don't work on windows

diff --git a/Engine/Build/BatchFiles/Linux/Toolchain/aarch64-unknown-linux-gnueabi.linux.config b/Engine/Build/BatchFiles/Linux/Toolchain/aarch64-unknown-linux-gnueabi.linux.config
--- a/Engine/Build/BatchFiles/Linux/Toolchain/aarch64-unknown-linux-gnueabi.linux.config
+++ b/Engine/Build/BatchFiles/Linux/Toolchain/aarch64-unknown-linux-gnueabi.linux.config
@@ -198,7 +198,7 @@ CT_ARCH_FLOAT_AUTO=y
 # CT_ARCH_FLOAT_HW is not set
 # CT_ARCH_FLOAT_SOFTFP is not set
 # CT_ARCH_FLOAT_SW is not set
-CT_TARGET_CFLAGS=""
+CT_TARGET_CFLAGS="-gz=none -Wl,--compress-debug-sections=none"
 CT_TARGET_LDFLAGS=""
 CT_ARCH_FLOAT="auto"
 # end of Target options
diff --git a/Engine/Build/BatchFiles/Linux/Toolchain/aarch64-unknown-linux-gnueabi.windows.config b/Engine/Build/BatchFiles/Linux/Toolchain/aarch64-unknown-linux-gnueabi.windows.config
--- a/Engine/Build/BatchFiles/Linux/Toolchain/aarch64-unknown-linux-gnueabi.windows.config
+++ b/Engine/Build/BatchFiles/Linux/Toolchain/aarch64-unknown-linux-gnueabi.windows.config
@@ -198,7 +198,7 @@ CT_ARCH_FLOAT_AUTO=y
 # CT_ARCH_FLOAT_HW is not set
 # CT_ARCH_FLOAT_SOFTFP is not set
 # CT_ARCH_FLOAT_SW is not set
-CT_TARGET_CFLAGS=""
+CT_TARGET_CFLAGS="-gz=none -Wl,--compress-debug-sections=none"
 CT_TARGET_LDFLAGS=""
 CT_ARCH_FLOAT="auto"
 # end of Target options
diff --git a/Engine/Build/BatchFiles/Linux/Toolchain/x86_64-unknown-linux-gnu.linux.config b/Engine/Build/BatchFiles/Linux/Toolchain/x86_64-unknown-linux-gnu.linux.config
--- a/Engine/Build/BatchFiles/Linux/Toolchain/x86_64-unknown-linux-gnu.linux.config
+++ b/Engine/Build/BatchFiles/Linux/Toolchain/x86_64-unknown-linux-gnu.linux.config
@@ -176,7 +176,7 @@ CT_ARCH_SUPPORTS_WITH_ARCH=y
 CT_ARCH_SUPPORTS_WITH_CPU=y
 CT_ARCH_SUPPORTS_WITH_TUNE=y
 CT_ARCH_ARCH=""
-CT_TARGET_CFLAGS=""
+CT_TARGET_CFLAGS="-gz=none -Wl,--compress-debug-sections=none"
 CT_TARGET_LDFLAGS=""
 # end of Target options
 
diff --git a/Engine/Build/BatchFiles/Linux/Toolchain/x86_64-unknown-linux-gnu.windows.config b/Engine/Build/BatchFiles/Linux/Toolchain/x86_64-unknown-linux-gnu.windows.config
--- a/Engine/Build/BatchFiles/Linux/Toolchain/x86_64-unknown-linux-gnu.windows.config
+++ b/Engine/Build/BatchFiles/Linux/Toolchain/x86_64-unknown-linux-gnu.windows.config
@@ -176,7 +176,7 @@ CT_ARCH_SUPPORTS_WITH_ARCH=y
 CT_ARCH_SUPPORTS_WITH_CPU=y
 CT_ARCH_SUPPORTS_WITH_TUNE=y
 CT_ARCH_ARCH=""
-CT_TARGET_CFLAGS=""
+CT_TARGET_CFLAGS="-gz=none -Wl,--compress-debug-sections=none"
 CT_TARGET_LDFLAGS=""
 # end of Target options
 
