Changed linux toolchain URL from http:// to https://
- Fixed SetupToolchain.sh to error out if tar extraction fails
- Test leftover toolchain archive for validity and remove if corrupted so that we don't continue to fail

[RN] Made SetupToolchain.sh a little more robust against download errors

diff --git a/Engine/Build/BatchFiles/Linux/SetupToolchain.sh b/Engine/Build/BatchFiles/Linux/SetupToolchain.sh
--- a/Engine/Build/BatchFiles/Linux/SetupToolchain.sh
+++ b/Engine/Build/BatchFiles/Linux/SetupToolchain.sh
@@ -25,7 +25,7 @@ fi
 
 TOOLCHAIN_ARCHIVE=$TOOLCHAIN_VERSION.tar.gz
 
-TOOLCHAIN_URL=http://cdn.unrealengine.com/Toolchain_Linux/native-linux-$TOOLCHAIN_ARCHIVE
+TOOLCHAIN_URL=https://cdn.unrealengine.com/Toolchain_Linux/native-linux-$TOOLCHAIN_ARCHIVE
 TOOLCHAIN_ROOT=Extras/ThirdPartyNotUE/SDKs/HostLinux/Linux_x64/
 TOOLCHAIN_CACHE="$GIT_DIR/ue4-sdks/"
 
@@ -51,6 +51,11 @@ else
 	TOOLCHAIN_ARCHIVE=$TOOLCHAIN_ROOT$TOOLCHAIN_ARCHIVE
 fi
 
+if ! tar tfz ${TOOLCHAIN_ARCHIVE} 2>/dev/null; then
+	echo Removing corrupt toolchain and redownloading
+	rm -f ${TOOLCHAIN_ARCHIVE}
+fi
+
 if [ ! -f $TOOLCHAIN_ARCHIVE ]; then
     rm -f ${TOOLCHAIN_ARCHIVE}.tmp
 	if which curl 1>/dev/null; then
@@ -68,7 +73,7 @@ fi
 
 if [ -f $TOOLCHAIN_ARCHIVE ]; then
 	echo "Extracting toolchain."
-	tar -xvf $TOOLCHAIN_ARCHIVE -C $TOOLCHAIN_ROOT 
+	tar -xvf $TOOLCHAIN_ARCHIVE -C $TOOLCHAIN_ROOT || exit 1
 	# If this is not a git build then do not cache the downloaded zip
 	if [ -f Build/PerforceBuild.txt ]; then
 		rm -f $TOOLCHAIN_ARCHIVE
