Don't allow multiple Build.bat processes to run at once from visual studio so we can properly build UBT
Detach UBT from always building before unreal projects

#rnx

diff --git a/Engine/Build/BatchFiles/Build.bat b/Engine/Build/BatchFiles/Build.bat
--- a/Engine/Build/BatchFiles/Build.bat
+++ b/Engine/Build/BatchFiles/Build.bat
@@ -13,6 +13,22 @@ rem ##     additional args are passed directly to UnrealBuildTool
 
 setlocal enabledelayedexpansion
 
+set waitingForLock=0
+
+:lock
+9>&2 2>nul (call :lockAndRestoreStdErr %* 9>"%tmp%\%~nx0.lock") || (
+	if %waitingForLock% == 0 set waitingForLock=1 && echo %~nx0 is already running, waiting for existing script to terminate...
+	ping 127.0.0.1 -n 2 >nul
+	goto :lock
+)
+exit /b
+
+:lockAndRestoreStdErr
+call :main %* 2>&9
+exit /b
+
+:main
+
 rem ## First, make sure the batch file exists in the folder we expect it to.  This is necessary in order to
 rem ## verify that our relative path to the /Engine/Source directory is correct
 if not exist "%~dp0..\..\Source" goto Error_BatchFileInWrongLocation
@@ -33,21 +49,8 @@ rem ## Compile UBT if the project file exists
 set ProjectFile="Programs\UnrealBuildTool\UnrealBuildTool.csproj"
 if not exist %ProjectFile% goto NoProjectFile
 
-rem ## Only build if UnrealBuildTool.dll is missing, as Visual Studio or GenerateProjectFiles should be building UnrealBuildTool
-rem ## Note: It is possible UnrealBuildTool will be out of date if the solution was generated with -NoDotNet or is VS2019
-rem ##       Historically this batch file did not compile UnrealBuildTool
-if not exist %UBTPath% (
-	rem ## If this script was called from Visual Studio 2022, build UBT with Visual Studio to prevent unnecessary rebuilds.
-	if "%VisualStudioVersion%" GEQ "17.0" (
-		echo Building UnrealBuildTool with %VisualStudioEdition%...
-		"%VSAPPIDDIR%..\..\MSBuild\Current\Bin\MSBuild.exe" %ProjectFile% -t:Build -p:Configuration=Development -verbosity:quiet -noLogo
-		if errorlevel 1 goto Error_UBTCompileFailed
-	) else (
-		echo Building UnrealBuildTool with dotnet...
-		dotnet build %ProjectFile% -c Development -v quiet
-		if errorlevel 1 goto Error_UBTCompileFailed
-	)
-)
+call ..\Build\BatchFiles\BuildUBT.bat
+if errorlevel 1 goto Error_UBTCompileFailed
 :NoProjectFile
 
 rem ## Run UBT
diff --git a/Engine/Source/Programs/UnrealBuildTool/ProjectFiles/ProjectFileGenerator.cs b/Engine/Source/Programs/UnrealBuildTool/ProjectFiles/ProjectFileGenerator.cs
--- a/Engine/Source/Programs/UnrealBuildTool/ProjectFiles/ProjectFileGenerator.cs
+++ b/Engine/Source/Programs/UnrealBuildTool/ProjectFiles/ProjectFileGenerator.cs
@@ -2075,9 +2075,6 @@ namespace UnrealBuildTool
 				VCSharpProjectFile UnrealBuildToolProject = new VCSharpProjectFile(ProjectFileName, Logger);
 				UnrealBuildToolProject.ShouldBuildForAllSolutionTargets = true;
 
-				// Store it off as we need it when generating target projects.
-				UBTProject = UnrealBuildToolProject;
-
 				// Add the project
 				AddExistingProjectFile(UnrealBuildToolProject, bNeedsAllPlatformAndConfigurations: true, bForceDevelopmentConfiguration: true);
 
