Don't allow multiple Build.bat processes to run at once from visual studio so we can properly build UBT
Detach UBT from always building before unreal projects

#rnx

diff --git a/Engine/Build/BatchFiles/Build.bat b/Engine/Build/BatchFiles/Build.bat
--- a/Engine/Build/BatchFiles/Build.bat
+++ b/Engine/Build/BatchFiles/Build.bat
@@ -13,6 +13,24 @@ rem ##     additional args are passed directly to UnrealBuildTool
 
 setlocal enabledelayedexpansion
 
+set ExitCode=
+set WaitingForLock=0
+
+:Lock
+9>&2 2>nul (call :LockAndRestoreStdErr %* 9>"%tmp%\%~nx0.lock") || (
+	if "!ExitCode!" NEQ "" exit /B !ExitCode!
+	if %WaitingForLock% == 0 set WaitingForLock=1 && echo %~nx0 is already running, waiting for existing script to terminate...
+	ping 127.0.0.1 -n 2 >nul
+	goto :Lock
+)
+exit /B !ExitCode!
+
+:LockAndRestoreStdErr
+call :Main %* 2>&9
+exit /B !ExitCode!
+
+:Main
+
 rem ## First, make sure the batch file exists in the folder we expect it to.  This is necessary in order to
 rem ## verify that our relative path to the /Engine/Source directory is correct
 if not exist "%~dp0..\..\Source" goto Error_BatchFileInWrongLocation
@@ -33,21 +51,8 @@ rem ## Compile UBT if the project file exists
 set ProjectFile="Programs\UnrealBuildTool\UnrealBuildTool.csproj"
 if not exist %ProjectFile% goto NoProjectFile
 
-rem ## Only build if UnrealBuildTool.dll is missing, as Visual Studio or GenerateProjectFiles should be building UnrealBuildTool
-rem ## Note: It is possible UnrealBuildTool will be out of date if the solution was generated with -NoDotNet or is VS2019
-rem ##       Historically this batch file did not compile UnrealBuildTool
-if not exist %UBTPath% (
-	rem ## If this script was called from Visual Studio 2022, build UBT with Visual Studio to prevent unnecessary rebuilds.
-	if "%VisualStudioVersion%" GEQ "17.0" (
-		echo Building UnrealBuildTool with %VisualStudioEdition%...
-		"%VSAPPIDDIR%..\..\MSBuild\Current\Bin\MSBuild.exe" %ProjectFile% -t:Build -p:Configuration=Development -verbosity:quiet -noLogo
-		if errorlevel 1 goto Error_UBTCompileFailed
-	) else (
-		echo Building UnrealBuildTool with dotnet...
-		dotnet build %ProjectFile% -c Development -v quiet
-		if errorlevel 1 goto Error_UBTCompileFailed
-	)
-)
+call ..\Build\BatchFiles\BuildUBT.bat
+if errorlevel 1 goto Error_UBTCompileFailed
 :NoProjectFile
 
 rem ## Run UBT
@@ -55,20 +60,25 @@ rem ## Run UBT
 if not exist %UBTPath% goto Error_UBTMissing
 echo Running UnrealBuildTool: dotnet %UBTPath% %*
 dotnet %UBTPath% %*
-EXIT /B !ERRORLEVEL!
+set ExitCode=!ERRORLEVEL!
+EXIT /B !ExitCode!
 
 :Error_BatchFileInWrongLocation
 echo ERROR: The batch file does not appear to be located in the Engine/Build/BatchFiles directory. This script must be run from within that directory.
-EXIT /B 999
+set ExitCode=999
+EXIT /B !ExitCode!
 
 :Error_NoDotnetSDK
 echo ERROR: Unable to find an install of Dotnet SDK.  Please make sure you have it installed and that `dotnet` is a globally available command.
-EXIT /B 999
+set ExitCode=999
+EXIT /B !ExitCode!
 
 :Error_UBTCompileFailed
 echo ERROR: Failed to build UnrealBuildTool.
-EXIT /B 999
+set ExitCode=999
+EXIT /B !ExitCode!
 
 :Error_UBTMissing
 echo ERROR: UnrealBuildTool.dll not found in %UBTPath%
-EXIT /B 999
+set ExitCode=999
+EXIT /B !ExitCode!
diff --git a/Engine/Source/Programs/UnrealBuildTool/ProjectFiles/ProjectFileGenerator.cs b/Engine/Source/Programs/UnrealBuildTool/ProjectFiles/ProjectFileGenerator.cs
--- a/Engine/Source/Programs/UnrealBuildTool/ProjectFiles/ProjectFileGenerator.cs
+++ b/Engine/Source/Programs/UnrealBuildTool/ProjectFiles/ProjectFileGenerator.cs
@@ -2075,9 +2075,6 @@ namespace UnrealBuildTool
 				VCSharpProjectFile UnrealBuildToolProject = new VCSharpProjectFile(ProjectFileName, Logger);
 				UnrealBuildToolProject.ShouldBuildForAllSolutionTargets = true;
 
-				// Store it off as we need it when generating target projects.
-				UBTProject = UnrealBuildToolProject;
-
 				// Add the project
 				AddExistingProjectFile(UnrealBuildToolProject, bNeedsAllPlatformAndConfigurations: true, bForceDevelopmentConfiguration: true);
 
