Fix error running UAT if the workspace root contains a space.

diff --git a/Engine/Build/BatchFiles/BuildUAT.bat b/Engine/Build/BatchFiles/BuildUAT.bat
--- a/Engine/Build/BatchFiles/BuildUAT.bat
+++ b/Engine/Build/BatchFiles/BuildUAT.bat
@@ -20,7 +20,7 @@ pushd "%~dp0..\..\Source"
 if not exist ..\Build\BatchFiles\BuildUAT.bat goto Error_BatchFileInWrongLocation
 
 rem Ensure UnrealBuildTool is up to date as it is a prereq
-call %~dp0\BuildUBT.bat
+call "%~dp0\BuildUBT.bat"
 if errorlevel 1 goto Error_UBTCompileFailed
 
 rem Check to see if the files in the AutomationTool, EpicGames.Build, EpicGames.Core, or UnrealBuildTool directory have changed.
