Build.bat: Quote BuildUBT.bat

#rnx

diff --git a/Engine/Build/BatchFiles/Build.bat b/Engine/Build/BatchFiles/Build.bat
--- a/Engine/Build/BatchFiles/Build.bat
+++ b/Engine/Build/BatchFiles/Build.bat
@@ -51,7 +51,7 @@ rem ## Compile UBT if the project file exists
 set ProjectFile="Programs\UnrealBuildTool\UnrealBuildTool.csproj"
 if not exist %ProjectFile% goto NoProjectFile
 
-call ..\Build\BatchFiles\BuildUBT.bat
+call "%~dp0\BuildUBT.bat"
 if errorlevel 1 goto Error_UBTCompileFailed
 :NoProjectFile
 
