BuildUAT.sh: Quote BuildUBT.sh

#rnx

diff --git a/Engine/Build/BatchFiles/BuildUAT.sh b/Engine/Build/BatchFiles/BuildUAT.sh
--- a/Engine/Build/BatchFiles/BuildUAT.sh
+++ b/Engine/Build/BatchFiles/BuildUAT.sh
@@ -20,7 +20,7 @@ if [ ! -f ../Build/BatchFiles/BuildUAT.sh ]; then
 fi
 
 # Ensure UnrealBuildTool is up to date as it is a prereq
-$SCRIPT_DIR/BuildUBT.sh
+"$SCRIPT_DIR/BuildUBT.sh"
 if [ $? -ne 0 ]; then
 popd >/dev/null
 exit 1
