Fix Build.bat for installed builds

#rnx

diff --git a/Engine/Build/BatchFiles/Build.bat b/Engine/Build/BatchFiles/Build.bat
--- a/Engine/Build/BatchFiles/Build.bat
+++ b/Engine/Build/BatchFiles/Build.bat
@@ -51,8 +51,19 @@ rem ## Compile UBT if the project file exists
 set ProjectFile="Programs\UnrealBuildTool\UnrealBuildTool.csproj"
 if not exist %ProjectFile% goto NoProjectFile
 
-call "%~dp0\BuildUBT.bat"
-if errorlevel 1 goto Error_UBTCompileFailed
+if exist "%~dp0\BuildUBT.bat" (
+	call "%~dp0\BuildUBT.bat"
+	if errorlevel 1 goto Error_UBTCompileFailed
+) ELSE (
+	rem ## Only build if UnrealBuildTool.dll is missing if BuildUBT.bat does not exist
+	if not exist %UBTPath% (
+		echo Building UnrealBuildTool with dotnet...
+		dotnet build %ProjectFile% -c Development -v quiet
+		dotnet publish %ProjectFile% -c Development --output ..\Binaries\DotNET\UnrealBuildTool --no-build -v quiet
+		if errorlevel 1 goto Error_UBTCompileFailed
+	)
+)
+
 :NoProjectFile
 
 rem ## Run UBT
