MSBuild task to write filelist and hashes to a depends file for all compiled inputs

#rnx

diff --git a/Engine/Build/BatchFiles/BuildUBT.bat b/Engine/Build/BatchFiles/BuildUBT.bat
--- a/Engine/Build/BatchFiles/BuildUBT.bat
+++ b/Engine/Build/BatchFiles/BuildUBT.bat
@@ -19,73 +19,35 @@ rem ## Change the CWD to /Engine/Source.
 pushd "%~dp0..\..\Source"
 if not exist ..\Build\BatchFiles\BuildUBT.bat goto Error_BatchFileInWrongLocation
 
-rem Check to see if the files in the UnrealBuildTool or Shared project directory have changed.
-rem find ".cs" files to only lines that match those names - excludes lines that will change for uninteresting reasons, like free space
-md ..\Intermediate\Build >nul 2>nul
-
-dir /s ^
- Programs\Shared\EpicGames.Build ^
- Programs\Shared\EpicGames.Core ^
- Programs\Shared\EpicGames.Horde ^
- Programs\Shared\EpicGames.IoHash ^
- Programs\Shared\EpicGames.MsBuild ^
- Programs\Shared\EpicGames.OIDC ^
- Programs\Shared\EpicGames.Serialization ^
- Programs\Shared\EpicGames.UBA ^
- Programs\Shared\EpicGames.UHT ^
- Programs\UnrealBuildTool ^
- 2>nul ^
- | %FIND% /c:"\.cs$" /c:"\.csproj$" > ..\Intermediate\Build\UnrealBuildToolFiles.txt
-
-dir /s ^
- ..\Binaries\Win64\UnrealBuildAccelerator ^
- 2>nul ^
- | %FIND% /c:"\.dll$" /c:"\.exe$" >> ..\Intermediate\Build\UnrealBuildToolFiles.txt
-
-if not exist ..\Platforms goto NoPlatforms
-for /d %%D in (..\Platforms\*) do (
-	dir /s ^
-	 %%D\Source\Programs\Shared ^
-	 %%D\Source\Programs\UnrealBuildTool ^
-	 2>nul ^
-	 | %FIND% /c:"\.cs$" /c:"\.csproj$" >> ..\Intermediate\Build\UnrealBuildToolFiles.txt
-)
-:NoPlatforms
-
-if not exist ..\Restricted goto NoRestricted
-for /d %%D in (..\Restricted\*) do (
-	dir /s ^
-	 %%D\Source\Programs\Shared ^
-	 %%D\Source\Programs\UnrealBuildTool ^
-	 2>nul ^
-	 | %FIND% /c:"\.cs$" /c:"\.csproj$" >> ..\Intermediate\Build\UnrealBuildToolFiles.txt
-)
-:NoRestricted
+rem ## Verify that dotnet is present
+call "%~dp0GetDotnetPath.bat"
+if errorlevel 1 goto Error_NoDotnetSDK
 
-rem note: no /s
-dir ^
- Programs\Shared\MetaData.cs ^
- 2>nul ^
- | %FIND% "\.cs$" >>..\Intermediate\Build\UnrealBuildToolFiles.txt
+rem ## Properties
+set DEPENDS_FILE=..\Intermediate\Build\UnrealBuildTool.dep.csv 
+set TEMP_DEPENDS_FILE="%TEMP%\UnrealBuildToolFiles-%random%.dep.csv"
 
+rem ## Command line arguments
 set MSBUILD_LOGLEVEL=%1
 if not defined %MSBUILD_LOGLEVEL set MSBUILD_LOGLEVEL=quiet
 
-set ARGUMENT=%2
-if not defined %ARGUMENT goto Check_UpToDate
-if /I "%ARGUMENT%" == "FORCE" goto Build_UnrealBuildTool
+rem Force build by deleting dependency file
+if /I "%2" == "FORCE" del %DEPENDS_FILE% >nul 2>&1
+
+rem Check to see if the files in the UnrealBuildTool solution have changed
+dotnet msbuild Programs\UnrealBuildTool\UnrealBuildTool.sln -t:ProjectDependencies -p:Configuration=Development -p:DependsPath=%TEMP_DEPENDS_FILE% -p:DependsOverwrite=false -m:1 -noLogo -v:%MSBUILD_LOGLEVEL%
+if errorlevel 1 goto Error_UBTDependsFailed
 
 :Check_UpToDate
 if not exist ..\Binaries\DotNET\UnrealBuildTool\UnrealBuildTool.dll goto Build_UnrealBuildTool
-set RUNUBT_EXITCODE=0
 rem per https://ss64.com/nt/fc.html using redirection syntax rather than errorlevel, based on observed inconsistent results from this function
-fc ..\Intermediate\Build\UnrealBuildToolFiles.txt ..\Intermediate\Build\UnrealBuildToolPrevFiles.txt >nul 2>&1 && goto Exit
+fc %TEMP_DEPENDS_FILE% %DEPENDS_FILE% >nul 2>&1 && (
+	del %TEMP_DEPENDS_FILE% >nul 2>&1
+	set RUNUBT_EXITCODE=0
+	goto Exit
+)
 
 :Build_UnrealBuildTool
-rem ## Verify that dotnet is present
-call "%~dp0GetDotnetPath.bat"
-if errorlevel 1 goto Error_NoDotnetSDK
-
 echo Building UnrealBuildTool...
 dotnet build Programs\UnrealBuildTool\UnrealBuildTool.csproj -c Development -v %MSBUILD_LOGLEVEL%
 if errorlevel 1 goto Error_UBTCompileFailed
@@ -95,7 +57,7 @@ dotnet publish Programs\UnrealBuildTool\UnrealBuildTool.csproj -c Development --
 if errorlevel 1 goto Error_UBTCompileFailed
 
 rem record input files - regardless of how we got here, these are now our point of reference
-copy /y ..\Intermediate\Build\UnrealBuildToolFiles.txt ..\Intermediate\Build\UnrealBuildToolPrevFiles.txt >nul
+move /y %TEMP_DEPENDS_FILE% %DEPENDS_FILE% >nul
 
 goto Exit
 
@@ -114,6 +76,13 @@ echo.
 set RUNUBT_EXITCODE=1
 goto Exit
 
+:Error_UBTDependsFailed
+echo.
+echo RunUBT ERROR: UnrealBuildTool failed to check dependencies.
+echo.
+set RUNUBT_EXITCODE=1
+goto Exit
+
 :Error_UBTCompileFailed
 echo.
 echo RunUBT ERROR: UnrealBuildTool failed to compile.
diff --git a/Engine/Source/Programs/Shared/EpicGames.Core.Tests/EpicGames.Core.Tests.csproj b/Engine/Source/Programs/Shared/EpicGames.Core.Tests/EpicGames.Core.Tests.csproj
--- a/Engine/Source/Programs/Shared/EpicGames.Core.Tests/EpicGames.Core.Tests.csproj
+++ b/Engine/Source/Programs/Shared/EpicGames.Core.Tests/EpicGames.Core.Tests.csproj
@@ -1,4 +1,5 @@
 <Project Sdk="Microsoft.NET.Sdk">
+  <Import Project="..\UnrealEngine.csproj.props" />
 
   <PropertyGroup>
     <TargetFramework>net8.0</TargetFramework>
diff --git a/Engine/Source/Programs/Shared/EpicGames.Core/EpicGames.Core.csproj b/Engine/Source/Programs/Shared/EpicGames.Core/EpicGames.Core.csproj
--- a/Engine/Source/Programs/Shared/EpicGames.Core/EpicGames.Core.csproj
+++ b/Engine/Source/Programs/Shared/EpicGames.Core/EpicGames.Core.csproj
@@ -1,4 +1,6 @@
 ﻿<Project Sdk="Microsoft.NET.Sdk">
+  <Import Project="..\UnrealEngine.csproj.props" />
+  
   <PropertyGroup>
     <TargetFramework>net8.0</TargetFramework>
     <Configuration Condition=" '$(Configuration)' == '' ">Development</Configuration>
diff --git a/Engine/Source/Programs/Shared/EpicGames.Horde.Tests/EpicGames.Horde.Tests.csproj b/Engine/Source/Programs/Shared/EpicGames.Horde.Tests/EpicGames.Horde.Tests.csproj
--- a/Engine/Source/Programs/Shared/EpicGames.Horde.Tests/EpicGames.Horde.Tests.csproj
+++ b/Engine/Source/Programs/Shared/EpicGames.Horde.Tests/EpicGames.Horde.Tests.csproj
@@ -1,4 +1,5 @@
 <Project Sdk="Microsoft.NET.Sdk">
+  <Import Project="..\UnrealEngine.csproj.props" />
 
   <PropertyGroup>
     <TargetFramework>net8.0</TargetFramework>
diff --git a/Engine/Source/Programs/Shared/EpicGames.Horde/EpicGames.Horde.csproj b/Engine/Source/Programs/Shared/EpicGames.Horde/EpicGames.Horde.csproj
--- a/Engine/Source/Programs/Shared/EpicGames.Horde/EpicGames.Horde.csproj
+++ b/Engine/Source/Programs/Shared/EpicGames.Horde/EpicGames.Horde.csproj
@@ -1,4 +1,6 @@
 <Project Sdk="Microsoft.NET.Sdk">
+  <Import Project="..\UnrealEngine.csproj.props" />
+
   <PropertyGroup>
     <TargetFramework>net8.0</TargetFramework>
     <Configuration Condition=" '$(Configuration)' == '' ">Development</Configuration>
diff --git a/Engine/Source/Programs/Shared/EpicGames.IoHash.Tests/EpicGames.IoHash.Tests.csproj b/Engine/Source/Programs/Shared/EpicGames.IoHash.Tests/EpicGames.IoHash.Tests.csproj
--- a/Engine/Source/Programs/Shared/EpicGames.IoHash.Tests/EpicGames.IoHash.Tests.csproj
+++ b/Engine/Source/Programs/Shared/EpicGames.IoHash.Tests/EpicGames.IoHash.Tests.csproj
@@ -1,4 +1,5 @@
 <Project Sdk="Microsoft.NET.Sdk">
+  <Import Project="..\UnrealEngine.csproj.props" />
 
   <PropertyGroup>
     <TargetFramework>net8.0</TargetFramework>
diff --git a/Engine/Source/Programs/Shared/EpicGames.IoHash/EpicGames.IoHash.csproj b/Engine/Source/Programs/Shared/EpicGames.IoHash/EpicGames.IoHash.csproj
--- a/Engine/Source/Programs/Shared/EpicGames.IoHash/EpicGames.IoHash.csproj
+++ b/Engine/Source/Programs/Shared/EpicGames.IoHash/EpicGames.IoHash.csproj
@@ -1,4 +1,6 @@
 ﻿<Project Sdk="Microsoft.NET.Sdk">
+  <Import Project="..\UnrealEngine.csproj.props" />
+  
   <PropertyGroup>
     <TargetFramework>net8.0</TargetFramework>
     <Configuration Condition=" '$(Configuration)' == '' ">Development</Configuration>
diff --git a/Engine/Source/Programs/Shared/EpicGames.OIDC/EpicGames.OIDC.csproj b/Engine/Source/Programs/Shared/EpicGames.OIDC/EpicGames.OIDC.csproj
--- a/Engine/Source/Programs/Shared/EpicGames.OIDC/EpicGames.OIDC.csproj
+++ b/Engine/Source/Programs/Shared/EpicGames.OIDC/EpicGames.OIDC.csproj
@@ -1,5 +1,6 @@
 <Project Sdk="Microsoft.NET.Sdk">
-
+  <Import Project="..\UnrealEngine.csproj.props" />
+  
   <PropertyGroup>
     <OutputType>Library</OutputType>
     <RootNamespace>EpicGames.OIDC</RootNamespace>
diff --git a/Engine/Source/Programs/Shared/EpicGames.Oodle/EpicGames.Oodle.csproj b/Engine/Source/Programs/Shared/EpicGames.Oodle/EpicGames.Oodle.csproj
--- a/Engine/Source/Programs/Shared/EpicGames.Oodle/EpicGames.Oodle.csproj
+++ b/Engine/Source/Programs/Shared/EpicGames.Oodle/EpicGames.Oodle.csproj
@@ -1,4 +1,6 @@
 <Project Sdk="Microsoft.NET.Sdk">
+  <Import Project="..\UnrealEngine.csproj.props" />
+  
   <PropertyGroup>
     <TargetFramework>net8.0</TargetFramework>
     <Configuration Condition=" '$(Configuration)' == '' ">Development</Configuration>
diff --git a/Engine/Source/Programs/Shared/EpicGames.Serialization.Tests/EpicGames.Serialization.Tests.csproj b/Engine/Source/Programs/Shared/EpicGames.Serialization.Tests/EpicGames.Serialization.Tests.csproj
--- a/Engine/Source/Programs/Shared/EpicGames.Serialization.Tests/EpicGames.Serialization.Tests.csproj
+++ b/Engine/Source/Programs/Shared/EpicGames.Serialization.Tests/EpicGames.Serialization.Tests.csproj
@@ -1,4 +1,5 @@
 <Project Sdk="Microsoft.NET.Sdk">
+  <Import Project="..\UnrealEngine.csproj.props" />
 
   <PropertyGroup>
     <TargetFramework>net8.0</TargetFramework>
diff --git a/Engine/Source/Programs/Shared/EpicGames.Serialization/EpicGames.Serialization.csproj b/Engine/Source/Programs/Shared/EpicGames.Serialization/EpicGames.Serialization.csproj
--- a/Engine/Source/Programs/Shared/EpicGames.Serialization/EpicGames.Serialization.csproj
+++ b/Engine/Source/Programs/Shared/EpicGames.Serialization/EpicGames.Serialization.csproj
@@ -1,4 +1,6 @@
 <Project Sdk="Microsoft.NET.Sdk">
+  <Import Project="..\UnrealEngine.csproj.props" />
+
   <PropertyGroup>
     <TargetFramework>net8.0</TargetFramework>
     <Configuration Condition=" '$(Configuration)' == '' ">Development</Configuration>
diff --git a/Engine/Source/Programs/UnrealBuildTool.Tests/UnrealBuildTool.Tests.csproj b/Engine/Source/Programs/UnrealBuildTool.Tests/UnrealBuildTool.Tests.csproj
--- a/Engine/Source/Programs/UnrealBuildTool.Tests/UnrealBuildTool.Tests.csproj
+++ b/Engine/Source/Programs/UnrealBuildTool.Tests/UnrealBuildTool.Tests.csproj
@@ -1,4 +1,5 @@
 <Project Sdk="Microsoft.NET.Sdk">
+  <Import Project="..\Shared\UnrealEngine.csproj.props" />
 
   <PropertyGroup>
     <TargetFramework>net8.0</TargetFramework>
