Fix exit code not being always returned

#rnx

diff --git a/Engine/Build/BatchFiles/Build.bat b/Engine/Build/BatchFiles/Build.bat
--- a/Engine/Build/BatchFiles/Build.bat
+++ b/Engine/Build/BatchFiles/Build.bat
@@ -16,6 +16,9 @@ setlocal enabledelayedexpansion
 set ExitCode=
 set WaitingForLock=0
 
+call :Lock %*
+exit /B !ExitCode!
+
 :Lock
 9>&2 2>nul (call :LockAndRestoreStdErr %* 9>"%tmp%\%~nx0.lock") || (
 	if "!ExitCode!" NEQ "" exit /B !ExitCode!
