Ensure ..\Binaries\DotNET\UnrealBuildTool exists before publishing

#rnx

diff --git a/Engine/Build/BatchFiles/BuildUBT.bat b/Engine/Build/BatchFiles/BuildUBT.bat
--- a/Engine/Build/BatchFiles/BuildUBT.bat
+++ b/Engine/Build/BatchFiles/BuildUBT.bat
@@ -53,6 +53,7 @@ dotnet build Programs\UnrealBuildTool\UnrealBuildTool.csproj -c Development -v %
 if errorlevel 1 goto Error_UBTCompileFailed
 
 echo Publishing UnrealBuildTool...
+mkdir ..\Binaries\DotNET\UnrealBuildTool  >nul 2>&1
 dotnet publish Programs\UnrealBuildTool\UnrealBuildTool.csproj -c Development --output ..\Binaries\DotNET\UnrealBuildTool --no-build -v %MSBUILD_LOGLEVEL%
 if errorlevel 1 goto Error_UBTCompileFailed
 
