Ensure ..\Intermediate\Build exists before publishing

#rnx

diff --git a/Engine/Build/BatchFiles/BuildUBT.bat b/Engine/Build/BatchFiles/BuildUBT.bat
--- a/Engine/Build/BatchFiles/BuildUBT.bat
+++ b/Engine/Build/BatchFiles/BuildUBT.bat
@@ -53,11 +53,12 @@ dotnet build Programs\UnrealBuildTool\UnrealBuildTool.csproj -c Development -v %
 if errorlevel 1 goto Error_UBTCompileFailed
 
 echo Publishing UnrealBuildTool...
-mkdir ..\Binaries\DotNET\UnrealBuildTool  >nul 2>&1
+if not exist ..\Binaries\DotNET\UnrealBuildTool mkdir ..\Binaries\DotNET\UnrealBuildTool >nul 2>&1
 dotnet publish Programs\UnrealBuildTool\UnrealBuildTool.csproj -c Development --output ..\Binaries\DotNET\UnrealBuildTool --no-build -v %MSBUILD_LOGLEVEL%
 if errorlevel 1 goto Error_UBTCompileFailed
 
 rem record input files - regardless of how we got here, these are now our point of reference
+if not exist ..\Intermediate\Build mkdir ..\Intermediate\Build >nul 2>&1
 move /y %TEMP_DEPENDS_FILE% %DEPENDS_FILE% >nul
 
 goto Exit
