BuildUBT: Update bash script to use dependency checking msbuild task

#rnx

diff --git a/Engine/Build/BatchFiles/BuildUBT.sh b/Engine/Build/BatchFiles/BuildUBT.sh
--- a/Engine/Build/BatchFiles/BuildUBT.sh
+++ b/Engine/Build/BatchFiles/BuildUBT.sh
@@ -8,17 +8,25 @@
 
 ## First, make sure the batch file exists in the folder we expect it to.  This is necessary in order to
 ## verify that our relative path to the /Engine/Source directory is correct
-SCRIPT_DIR=$(cd "`dirname "$0"`" && pwd)
-pushd "$SCRIPT_DIR/../../Source" >/dev/null
+SCRIPT_DIR=$(cd "$(dirname "$0")" && pwd)
+pushd "$SCRIPT_DIR/../../Source" >/dev/null || exit 1
 
 if [ ! -f ../Build/BatchFiles/BuildUBT.sh ]; then
   echo
   echo "BuildUBT ERROR: The script does not appear to be located in the /Engine/Build/BatchFiles directory.  This script must be run from within that directory."
   echo
-  popd >/dev/null
+  popd >/dev/null || exit 1
   exit 1
 fi
 
+if [ "$(uname)" = "Darwin" ]; then
+  # Setup Environment
+  source "$SCRIPT_DIR/Mac/SetupEnvironment.sh" -dotnet "$SCRIPT_DIR/Mac"
+elif [ "$(uname)" = "Linux" ]; then
+  # Setup Environment
+  source "$SCRIPT_DIR/Linux/SetupEnvironment.sh" -dotnet "$SCRIPT_DIR/Linux"
+fi
+
 MSBuild_Verbosity="${1:-quiet}"
 
 # Check to see if the files in the UnrealBuildTool, EpicGames.Build, EpicGames.Core, or UnrealBuildTool
@@ -26,64 +34,24 @@ MSBuild_Verbosity="${1:-quiet}"
 
 mkdir -p ../Intermediate/Build
 
+DEPENDS_FILE=../Intermediate/Build/UnrealBuildTool.dep.csv
+TEMP_DEPENDS_FILE=$(mktemp)
+
+dotnet msbuild Programs/UnrealBuildTool/UnrealBuildTool.sln -t:ProjectDependencies -p:Configuration=Development -p:DependsPath="$TEMP_DEPENDS_FILE" -p:DependsOverwrite=false -m:1 -noLogo -v:"$MSBuild_Verbosity"
+
 PERFORM_REBUILD=0
-if [ "$1" == "FORCE" ]; then
+if [ "$2" == "FORCE" ]; then
   PERFORM_REBUILD=1
-  echo "Rebuilding: build requested"
+  echo "Rebuilding: Force build requested"
 
 elif [ ! -f ../Binaries/DotNET/UnrealBuildTool/UnrealBuildTool.dll ]; then
   PERFORM_REBUILD=1
   echo "Rebuilding: UnrealBuildTool assembly not found"
 
-elif [ -f ../Intermediate/Build/UnrealBuildToolLastBuildTime ]; then
-  UPDATED_DEP_FILES="$(find \
-    Programs/Shared/EpicGames.Build \
-    Programs/Shared/EpicGames.Core \
-    Programs/Shared/EpicGames.Horde \
-    Programs/Shared/EpicGames.IoHash \
-    Programs/Shared/EpicGames.MsBuild \
-    Programs/Shared/EpicGames.OIDC \
-    Programs/Shared/EpicGames.Serialization \
-    Programs/Shared/EpicGames.UBA \
-    Programs/Shared/EpicGames.UHT \
-    Programs/UnrealBuildTool \
-    ../Restricted/**/Source/Programs/Shared \
-    ../Platforms/*/Source/Programs/Shared \
-    ../Restricted/**/Source/Programs/UnrealBuildTool \
-    ../Platforms/*/Source/Programs/UnrealBuildTool \
-    -type f \
-    \( -iname \*.cs -or -iname \*.csproj \) \
-    -newer ../Intermediate/Build/UnrealBuildToolLastBuildTime)"
-  
-  if [ -n "$UPDATED_DEP_FILES" ]; then
+elif [ -f ../Intermediate/Build/UnrealBuildTool.dep.csv ]; then
+  if ! cmp --silent $DEPENDS_FILE "$TEMP_DEPENDS_FILE"; then
     PERFORM_REBUILD=1
-    echo "Rebuilding: found updated files:"
-    echo "$UPDATED_DEP_FILES"
-  fi
-
-  UPDATED_DEP_FILES="$(find \
-    Programs/Shared \
-    -maxdepth 1 \
-    -type f \
-    \( -iname \*.cs -or -iname \*.csproj \) \
-    -newer ../Intermediate/Build/UnrealBuildToolLastBuildTime)"
-  if [ -n "$UPDATED_DEP_FILES" ]; then
-    PERFORM_REBUILD=1
-    echo "Rebuilding: Found updated files:"
-    echo "$UPDATED_DEP_FILES"
-  fi
-
-  UPDATED_DEP_FILES="$(find \
-    ../Binaries/Linux* \
-    ../Binaries/Mac \
-    -maxdepth 2 \
-    -type f \
-    \( -iname \*Uba* \) \
-    -newer ../Intermediate/Build/UnrealBuildToolLastBuildTime)"
-  if [ -n "$UPDATED_DEP_FILES" ]; then
-    PERFORM_REBUILD=1
-    echo "Rebuilding: Found updated files:"
-    echo "$UPDATED_DEP_FILES"
+    echo "Rebuilding: Found updated files"
   fi
 
 else
@@ -92,35 +60,26 @@ else
 fi
 
 if [ $PERFORM_REBUILD -eq 1 ]; then
-  if [ "$(uname)" = "Darwin" ]; then
-    # Setup Environment
-    source "$SCRIPT_DIR/Mac/SetupEnvironment.sh" -dotnet "$SCRIPT_DIR/Mac"
-  fi
-
-  if [ "$(uname)" = "Linux" ]; then
-    # Setup Environment
-    source "$SCRIPT_DIR/Linux/SetupEnvironment.sh" -dotnet "$SCRIPT_DIR/Linux"
-  fi
-
-  dotnet build Programs/UnrealBuildTool/UnrealBuildTool.csproj -c Development -v $MSBuild_Verbosity
+  echo "Building UnrealBuildTool..."
+  dotnet build Programs/UnrealBuildTool/UnrealBuildTool.csproj -c Development -v "$MSBuild_Verbosity"
   if [ $? -ne 0 ]; then
-    echo "Compilation failed"
-    popd >/dev/null
+    echo "UnrealBuildTool compilation failed"
+    popd >/dev/null || exit 1
     exit 1
   fi
 
-  dotnet publish Programs/UnrealBuildTool/UnrealBuildTool.csproj -c Development --output ../Binaries/DotNET/UnrealBuildTool --no-build -v $MSBuild_Verbosity
+  echo "Publishing UnrealBuildTool..."
+  dotnet publish Programs/UnrealBuildTool/UnrealBuildTool.csproj -c Development --output ../Binaries/DotNET/UnrealBuildTool --no-build -v "$MSBuild_Verbosity"
   if [ $? -ne 0 ]; then
     echo "UnrealBuildTool publish failed"
-    popd >/dev/null
+    popd >/dev/null || exit 1
     exit 1
   fi
 
-  touch ../Intermediate/Build/UnrealBuildToolLastBuildTime
+  mv -f "$TEMP_DEPENDS_FILE" $DEPENDS_FILE >/dev/null
 else
-  echo "UnrealBuildTool.dll is up to date"
+  echo "UnrealBuildTool is up to date"
 fi
 
-popd >/dev/null
+popd >/dev/null || exit 1
 exit 0
-
