Update system dotnet requirement in batch

#rnx

diff --git a/Engine/Build/BatchFiles/GetDotnetPath.bat b/Engine/Build/BatchFiles/GetDotnetPath.bat
--- a/Engine/Build/BatchFiles/GetDotnetPath.bat
+++ b/Engine/Build/BatchFiles/GetDotnetPath.bat
@@ -45,7 +45,7 @@ exit /B 1
 
 :find_sdks
 
-set REQUIRED_MAJOR_VERSION=6
+set REQUIRED_MAJOR_VERSION=8
 set REQUIRED_MINOR_VERSION=0
 
 set FOUND_MAJOR=
