BuildUBT: sort depends file

#rnx

diff --git a/Engine/Build/BatchFiles/BuildUBT.bat b/Engine/Build/BatchFiles/BuildUBT.bat
--- a/Engine/Build/BatchFiles/BuildUBT.bat
+++ b/Engine/Build/BatchFiles/BuildUBT.bat
@@ -25,7 +25,7 @@ if errorlevel 1 goto Error_NoDotnetSDK
 
 rem ## Properties
 set DEPENDS_FILE=..\Intermediate\Build\UnrealBuildTool.dep.csv 
-set TEMP_DEPENDS_FILE="%TEMP%\UnrealBuildToolFiles-%random%.dep.csv"
+set TEMP_DEPENDS_FILE=%TEMP%\UnrealBuildToolFiles-%random%.dep.csv
 
 rem ## Command line arguments
 set MSBUILD_LOGLEVEL=%1
@@ -37,6 +37,8 @@ if /I "%2" == "FORCE" del %DEPENDS_FILE% >nul 2>&1
 rem Check to see if the files in the UnrealBuildTool solution have changed
 dotnet msbuild Programs\UnrealBuildTool\UnrealBuildTool.sln -t:ProjectDependencies -p:Configuration=Development -p:DependsPath=%TEMP_DEPENDS_FILE% -p:DependsOverwrite=false -m:1 -noLogo -v:%MSBUILD_LOGLEVEL%
 if errorlevel 1 goto Error_UBTDependsFailed
+sort /unique /o %TEMP_DEPENDS_FILE% %TEMP_DEPENDS_FILE%
+if errorlevel 1 goto Error_UBTDependsFailed
 
 :Check_UpToDate
 if not exist ..\Binaries\DotNET\UnrealBuildTool\UnrealBuildTool.dll goto Build_UnrealBuildTool
diff --git a/Engine/Build/BatchFiles/BuildUBT.sh b/Engine/Build/BatchFiles/BuildUBT.sh
--- a/Engine/Build/BatchFiles/BuildUBT.sh
+++ b/Engine/Build/BatchFiles/BuildUBT.sh
@@ -38,6 +38,7 @@ DEPENDS_FILE=../Intermediate/Build/UnrealBuildTool.dep.csv
 TEMP_DEPENDS_FILE=$(mktemp)
 
 dotnet msbuild Programs/UnrealBuildTool/UnrealBuildTool.sln -t:ProjectDependencies -p:Configuration=Development -p:DependsPath="$TEMP_DEPENDS_FILE" -p:DependsOverwrite=false -m:1 -noLogo -v:"$MSBuild_Verbosity"
+sort -u -o "$TEMP_DEPENDS_FILE" "$TEMP_DEPENDS_FILE"
 
 PERFORM_REBUILD=0
 if [ "$2" == "FORCE" ]; then
