BuildUBT: pass platform to dotnet msbuild in case of environment variables

#rnx

diff --git a/Engine/Build/BatchFiles/BuildUBT.bat b/Engine/Build/BatchFiles/BuildUBT.bat
--- a/Engine/Build/BatchFiles/BuildUBT.bat
+++ b/Engine/Build/BatchFiles/BuildUBT.bat
@@ -35,7 +35,7 @@ rem Force build by deleting dependency file
 if /I "%2" == "FORCE" del %DEPENDS_FILE% >nul 2>&1
 
 rem Check to see if the files in the UnrealBuildTool solution have changed
-dotnet msbuild Programs\UnrealBuildTool\UnrealBuildTool.sln -t:ProjectDependencies -p:Configuration=Development -p:DependsPath=%TEMP_DEPENDS_FILE% -p:DependsOverwrite=false -m:1 -noLogo -v:%MSBUILD_LOGLEVEL%
+dotnet msbuild Programs\UnrealBuildTool\UnrealBuildTool.sln -t:ProjectDependencies -p:Configuration=Development -p:Platform="Any CPU" -p:DependsPath=%TEMP_DEPENDS_FILE% -p:DependsOverwrite=false -m:1 -noLogo -v:%MSBUILD_LOGLEVEL%
 if errorlevel 1 goto Error_UBTDependsFailed
 sort /unique /o %TEMP_DEPENDS_FILE% %TEMP_DEPENDS_FILE%
 if errorlevel 1 goto Error_UBTDependsFailed
