BuildUBT: Full path to sort.exe

#rnx

diff --git a/Engine/Build/BatchFiles/BuildUBT.bat b/Engine/Build/BatchFiles/BuildUBT.bat
--- a/Engine/Build/BatchFiles/BuildUBT.bat
+++ b/Engine/Build/BatchFiles/BuildUBT.bat
@@ -8,9 +8,6 @@ rem ## if you copy it to a different location and run it.
 
 setlocal
 
-rem ## Make sure we use the Windows "find" utility and not a Unix-like tool found in PATH
-set FIND=%SYSTEMROOT%\System32\findstr.exe /i /r
-
 rem ## First, make sure the batch file exists in the folder we expect it to.  This is necessary in order to
 rem ## verify that our relative path to the /Engine/Source directory is correct
 if not exist "%~dp0..\..\Source" goto Error_BatchFileInWrongLocation
@@ -37,7 +34,7 @@ if /I "%2" == "FORCE" del %DEPENDS_FILE% >nul 2>&1
 rem Check to see if the files in the UnrealBuildTool solution have changed
 dotnet msbuild Programs\UnrealBuildTool\UnrealBuildTool.sln -t:ProjectDependencies -p:Configuration=Development -p:Platform="Any CPU" -p:DependsPath=%TEMP_DEPENDS_FILE% -p:DependsOverwrite=false -m:1 -noLogo -v:%MSBUILD_LOGLEVEL%
 if errorlevel 1 goto Error_UBTDependsFailed
-sort /unique /o %TEMP_DEPENDS_FILE% %TEMP_DEPENDS_FILE%
+"%SYSTEMROOT%\System32\sort.exe" /unique /o %TEMP_DEPENDS_FILE% %TEMP_DEPENDS_FILE%
 if errorlevel 1 goto Error_UBTDependsFailed
 
 :Check_UpToDate
