Build.bat: Allow running multiple instances from different folder paths to account for multiple branches/installs

#rnx

diff --git a/Engine/Build/BatchFiles/Build.bat b/Engine/Build/BatchFiles/Build.bat
--- a/Engine/Build/BatchFiles/Build.bat
+++ b/Engine/Build/BatchFiles/Build.bat
@@ -15,12 +15,15 @@ setlocal enabledelayedexpansion
 
 set ExitCode=
 set WaitingForLock=0
+set LockFile=%~f0
+set LockFile=%LockFile:\=-%
+set LockFile=%tmp%\%LockFile::=%.lock
 
 call :Lock %*
 exit /B !ExitCode!
 
 :Lock
-9>&2 2>nul (call :LockAndRestoreStdErr %* 9>"%tmp%\%~nx0.lock") || (
+9>&2 2>nul (call :LockAndRestoreStdErr %* 9>"%LockFile%") || (
 	if "!ExitCode!" NEQ "" exit /B !ExitCode!
 	if %WaitingForLock% == 0 set WaitingForLock=1 && echo %~nx0 is already running, waiting for existing script to terminate...
 	ping 127.0.0.1 -n 2 >nul
