Added a way to copy files into resigned app

diff --git a/Engine/Build/BatchFiles/Mac/ResignApp.sh b/Engine/Build/BatchFiles/Mac/ResignApp.sh
--- a/Engine/Build/BatchFiles/Mac/ResignApp.sh
+++ b/Engine/Build/BatchFiles/Mac/ResignApp.sh
@@ -2,6 +2,8 @@
 
 set -e
 
+FILES_TO_COPY=()
+
 DoWork()
 {
 	echo "------------------------------------------------------------------------------"
@@ -67,6 +69,13 @@ DoWork()
 		echo "Setting commandline to $CMDLINE"
 		echo $CMDLINE > "$TARGET/uecommandline.txt"
 	fi
+ 
+    for COPY_FILE_PATH in "${FILES_TO_COPY[@]}"
+    do
+        COPY_FILE_NAME=$(basename ${COPY_FILE_PATH});
+        echo "Copying $COPY_FILE_PATH to $TARGET/$COPY_FILE_NAME"
+        cp $COPY_FILE_PATH "$TARGET/$COPY_FILE_NAME"
+    done
 
 	echo "------------------------------------------------------------------------------"
 
@@ -165,6 +174,8 @@ Help()
 	echo "     Bundle version for the app. This is the internal version string that are gathered within one Marketing Version in TF/ASC"
 	echo "  -c | --cmdline"
 	echo "     A commandline to place into the app as uecommandline.txt"
+    echo "  -cpy | --copy"
+    echo "     Additional file to copy into the app, can be used multiple times. Copied files will be placed into the root directiory, without subdirectories"
 	echo "  -h | --help"
 	echo "     Show this message"
 	exit 0
@@ -218,6 +229,11 @@ while [[ $# -gt 0 ]]; do
 	  shift # past argument
 	  shift # past value
 	  ;;
+    -cpy|--copy)
+      FILES_TO_COPY+=("$2")
+      shift # past argument
+      shift # past value
+      ;;
 	-h|--help)
 	  Help
 	  shift # past value
