BuildUBT: Fix cannot access dep.csv

#rnx

diff --git a/Engine/Build/BatchFiles/BuildUBT.bat b/Engine/Build/BatchFiles/BuildUBT.bat
--- a/Engine/Build/BatchFiles/BuildUBT.bat
+++ b/Engine/Build/BatchFiles/BuildUBT.bat
@@ -22,7 +22,8 @@ if errorlevel 1 goto Error_NoDotnetSDK
 
 rem ## Properties
 set DEPENDS_FILE=..\Intermediate\Build\UnrealBuildTool.dep.csv 
-set TEMP_DEPENDS_FILE=%TEMP%\UnrealBuildToolFiles-%random%.dep.csv
+set TEMP_DEPENDS_PATH=%TEMP%\UnrealBuildToolFiles-%random%%random%
+set TEMP_DEPENDS_FILE=%TEMP_DEPENDS_PATH%.dep.csv
 
 rem ## Command line arguments
 set MSBUILD_LOGLEVEL=%1
@@ -32,8 +33,11 @@ rem Force build by deleting dependency file
 if /I "%2" == "FORCE" del %DEPENDS_FILE% >nul 2>&1
 
 rem Check to see if the files in the UnrealBuildTool solution have changed
-dotnet msbuild Programs\UnrealBuildTool\UnrealBuildTool.sln -t:ProjectDependencies -p:Configuration=Development -p:Platform="Any CPU" -p:DependsPath=%TEMP_DEPENDS_FILE% -p:DependsOverwrite=false -m:1 -noLogo -v:%MSBUILD_LOGLEVEL%
+mkdir %TEMP_DEPENDS_PATH%
+dotnet msbuild Programs\UnrealBuildTool\UnrealBuildTool.sln -t:ProjectDependencies -p:Configuration=Development -p:Platform="Any CPU" -p:OutputPath=%TEMP_DEPENDS_PATH% -noLogo -v:%MSBUILD_LOGLEVEL%
 if errorlevel 1 goto Error_UBTDependsFailed
+copy /b %TEMP_DEPENDS_PATH%\*.dep.csv %TEMP_DEPENDS_FILE% >nul 2>&1
+rmdir /s /q %TEMP_DEPENDS_PATH% >nul 2>&1
 "%SYSTEMROOT%\System32\sort.exe" /unique /o %TEMP_DEPENDS_FILE% %TEMP_DEPENDS_FILE%
 if errorlevel 1 goto Error_UBTDependsFailed
 
