BuildUBT: random is not actually random, use time instead

#rnx

diff --git a/Engine/Build/BatchFiles/BuildUBT.bat b/Engine/Build/BatchFiles/BuildUBT.bat
--- a/Engine/Build/BatchFiles/BuildUBT.bat
+++ b/Engine/Build/BatchFiles/BuildUBT.bat
@@ -22,7 +22,7 @@ if errorlevel 1 goto Error_NoDotnetSDK
 
 rem ## Properties
 set DEPENDS_FILE=..\Intermediate\Build\UnrealBuildTool.dep.csv 
-set TEMP_DEPENDS_PATH=%TEMP%\UnrealBuildToolFiles-%random%%random%
+set TEMP_DEPENDS_PATH=%TEMP%\UnrealBuildToolFiles-%time::=-%
 set TEMP_DEPENDS_FILE=%TEMP_DEPENDS_PATH%.dep.csv
 
 rem ## Command line arguments
@@ -33,7 +33,7 @@ rem Force build by deleting dependency file
 if /I "%2" == "FORCE" del %DEPENDS_FILE% >nul 2>&1
 
 rem Check to see if the files in the UnrealBuildTool solution have changed
-mkdir %TEMP_DEPENDS_PATH%
+mkdir %TEMP_DEPENDS_PATH% >nul 2>&1
 dotnet msbuild Programs\UnrealBuildTool\UnrealBuildTool.sln -t:ProjectDependencies -p:Configuration=Development -p:Platform="Any CPU" -p:OutputPath=%TEMP_DEPENDS_PATH% -noLogo -v:%MSBUILD_LOGLEVEL%
 if errorlevel 1 goto Error_UBTDependsFailed
 copy /b %TEMP_DEPENDS_PATH%\*.dep.csv %TEMP_DEPENDS_FILE% >nul 2>&1
