BuildUBT: handle commas

#rnx

diff --git a/Engine/Build/BatchFiles/BuildUBT.bat b/Engine/Build/BatchFiles/BuildUBT.bat
--- a/Engine/Build/BatchFiles/BuildUBT.bat
+++ b/Engine/Build/BatchFiles/BuildUBT.bat
@@ -23,6 +23,7 @@ if errorlevel 1 goto Error_NoDotnetSDK
 rem ## Properties
 set DEPENDS_FILE=..\Intermediate\Build\UnrealBuildTool.dep.csv 
 set TEMP_DEPENDS_PATH=%TEMP%\UnrealBuildToolFiles-%time::=-%
+set TEMP_DEPENDS_PATH=%TEMP_DEPENDS_PATH:,=.%
 set TEMP_DEPENDS_FILE=%TEMP_DEPENDS_PATH%.dep.csv
 
 rem ## Command line arguments
