BuildUBT: Use a retry for the temp path

#rnx

diff --git a/Engine/Build/BatchFiles/BuildUBT.bat b/Engine/Build/BatchFiles/BuildUBT.bat
--- a/Engine/Build/BatchFiles/BuildUBT.bat
+++ b/Engine/Build/BatchFiles/BuildUBT.bat
@@ -22,7 +22,9 @@ if errorlevel 1 goto Error_NoDotnetSDK
 
 rem ## Properties
 set DEPENDS_FILE=..\Intermediate\Build\UnrealBuildTool.dep.csv 
-set TEMP_DEPENDS_PATH=%TEMP%\UnrealBuildToolFiles-%random%%random%
+:tempRetry
+set TEMP_DEPENDS_PATH=%TEMP%\BuildUBT-%random%%random%
+if exist %TEMP_DEPENDS_PATH% goto :tempRetry
 set TEMP_DEPENDS_FILE=%TEMP_DEPENDS_PATH%.dep.csv
 
 rem ## Command line arguments
@@ -33,19 +35,17 @@ rem Force build by deleting dependency file
 if /I "%2" == "FORCE" del %DEPENDS_FILE% >nul 2>&1
 
 rem Check to see if the files in the UnrealBuildTool solution have changed
-mkdir %TEMP_DEPENDS_PATH%
-dotnet msbuild Programs\UnrealBuildTool\UnrealBuildTool.sln -t:ProjectDependencies -p:Configuration=Development -p:Platform="Any CPU" -p:OutputPath=%TEMP_DEPENDS_PATH% -noLogo -v:%MSBUILD_LOGLEVEL%
+mkdir "%TEMP_DEPENDS_PATH%"
+dotnet msbuild Programs\UnrealBuildTool\UnrealBuildTool.sln -t:ProjectDependencies -p:Configuration=Development -p:Platform="Any CPU" -p:OutputPath="%TEMP_DEPENDS_PATH%" -noLogo -v:%MSBUILD_LOGLEVEL%
 if errorlevel 1 goto Error_UBTDependsFailed
-copy /b %TEMP_DEPENDS_PATH%\*.dep.csv %TEMP_DEPENDS_FILE% >nul 2>&1
-rmdir /s /q %TEMP_DEPENDS_PATH% >nul 2>&1
-"%SYSTEMROOT%\System32\sort.exe" /unique /o %TEMP_DEPENDS_FILE% %TEMP_DEPENDS_FILE%
+copy /b "%TEMP_DEPENDS_PATH%\*.dep.csv" "%TEMP_DEPENDS_FILE%" >nul 2>&1
+"%SYSTEMROOT%\System32\sort.exe" /unique /o "%TEMP_DEPENDS_FILE%" "%TEMP_DEPENDS_FILE%"
 if errorlevel 1 goto Error_UBTDependsFailed
 
 :Check_UpToDate
 if not exist ..\Binaries\DotNET\UnrealBuildTool\UnrealBuildTool.dll goto Build_UnrealBuildTool
 rem per https://ss64.com/nt/fc.html using redirection syntax rather than errorlevel, based on observed inconsistent results from this function
-fc %TEMP_DEPENDS_FILE% %DEPENDS_FILE% >nul 2>&1 && (
-	del %TEMP_DEPENDS_FILE% >nul 2>&1
+fc "%TEMP_DEPENDS_FILE%" %DEPENDS_FILE% >nul 2>&1 && (
 	set RUNUBT_EXITCODE=0
 	goto Exit
 )
@@ -62,7 +62,7 @@ if errorlevel 1 goto Error_UBTCompileFailed
 
 rem record input files - regardless of how we got here, these are now our point of reference
 if not exist ..\Intermediate\Build mkdir ..\Intermediate\Build >nul 2>&1
-move /y %TEMP_DEPENDS_FILE% %DEPENDS_FILE% >nul
+move /y "%TEMP_DEPENDS_FILE%" %DEPENDS_FILE% >nul
 
 goto Exit
 
@@ -96,6 +96,9 @@ set RUNUBT_EXITCODE=1
 goto Exit
 
 :Exit
+rem ## Clean up temp files
+del "%TEMP_DEPENDS_FILE%" >nul 2>&1
+rmdir /s /q "%TEMP_DEPENDS_PATH%" >nul 2>&1
 rem ## Restore original CWD in case we change it
 popd
 exit /B %RUNUBT_EXITCODE%
