BuildUBT: case insensitve compare

#rnx

diff --git a/Engine/Build/BatchFiles/BuildUBT.bat b/Engine/Build/BatchFiles/BuildUBT.bat
--- a/Engine/Build/BatchFiles/BuildUBT.bat
+++ b/Engine/Build/BatchFiles/BuildUBT.bat
@@ -45,7 +45,7 @@ if errorlevel 1 goto Error_UBTDependsFailed
 :Check_UpToDate
 if not exist ..\Binaries\DotNET\UnrealBuildTool\UnrealBuildTool.dll goto Build_UnrealBuildTool
 rem per https://ss64.com/nt/fc.html using redirection syntax rather than errorlevel, based on observed inconsistent results from this function
-fc "%TEMP_DEPENDS_FILE%" %DEPENDS_FILE% >nul 2>&1 && (
+fc /C "%TEMP_DEPENDS_FILE%" %DEPENDS_FILE% >nul 2>&1 && (
 	set RUNUBT_EXITCODE=0
 	goto Exit
 )
