Fix for universal dsym lipo fail due to size

#rnx

diff --git a/Engine/Build/BatchFiles/Mac/GenerateUniversalDSYM.sh b/Engine/Build/BatchFiles/Mac/GenerateUniversalDSYM.sh
--- a/Engine/Build/BatchFiles/Mac/GenerateUniversalDSYM.sh
+++ b/Engine/Build/BatchFiles/Mac/GenerateUniversalDSYM.sh
@@ -60,23 +60,8 @@ if [ ${#arches[@]} -gt 1 ]; then
   
 	# lipo all dsym binaries directly to the project-named file that we'd expect
 	echo "Merging architectures '${binpaths[*]}' together into $2"
-
-	set +e 	
-	
-	lipo ${binpaths[*]} -create -output "$2/Contents/Resources/DWARF/${binaryname}"
-		
-	if [[ $? != 0 ]]; then 
-		set -e	
-		echo ""
-		echo ======================================================================================
-		echo  Debug info too large for standard dSYM. Using FAT64 format, which will not work 
-		echo  with lldb before Xcode 15. Other tools that use the dSYM may fail as well.
-		echo ======================================================================================
-		
-		lipo ${binpaths[*]} -fat64 -create -output "$2/Contents/Resources/DWARF/${binaryname}"
-	fi
-	set -e	
 	
+	lipo ${binpaths[*]} -fat64 -create -output "$2/Contents/Resources/DWARF/${binaryname}"
 	
 	retVal=$?
 	if [ $retVal -ne 0 ]; then
