Invoke BuildUBT from XcodeBuild instead of building the CSproj directly. Previous implementation built the UBT project which didn't deploy the dll to its final use destination for subsequent build steps. Additionally, fix XcodeBuild's trailing argument forwarding such that an empty string argument isn't forwarded to UBT when there aren't any trailiing arguments in the command line invocation of XcodeBuild.

diff --git a/Engine/Build/BatchFiles/Mac/XcodeBuild.sh b/Engine/Build/BatchFiles/Mac/XcodeBuild.sh
--- a/Engine/Build/BatchFiles/Mac/XcodeBuild.sh
+++ b/Engine/Build/BatchFiles/Mac/XcodeBuild.sh
@@ -39,7 +39,7 @@ if [ $BUILD_UBT == 1 ]; then
 	unset TARGETNAME
 	
 	if [ ! -f Engine/Build/InstalledBuild.txt ]; then
-		dotnet build Engine/Source/Programs/UnrealBuildTool/UnrealBuildTool.csproj -c Development -v quiet
+		$(dirname "$0")/../BuildUBT.sh
 
 		if [ $? -ne 0 ]; then
 			echo "Failed to build the build tool (UnrealBuildTool)"
@@ -55,13 +55,13 @@ if [ ["$ACTION"] == [""] ]; then
 	TARGET=$1
 	PLATFORM=$2
 	CONFIGURATION=$3
-	TRAILINGARGS=${@:4}
+	TRAILINGARGS=("${@:4}")
 else
 	# non build actions are all shifted by one
 	TARGET=$2
 	PLATFORM=$3
 	CONFIGURATION=$4
-	TRAILINGARGS=${@:5}
+	TRAILINGARGS=("${@:5}")
 fi
 
 if [[ $ARCHS ]]; then
@@ -95,7 +95,7 @@ case $PLATFORM in
 	;;
 esac
 
-echo "Processing $ACTION for Target=$TARGET Platform=$PLATFORM Configuration=$CONFIGURATION $UBT_ARCHFLAG $TRAILINGARGS "
+echo "Processing $ACTION for Target=$TARGET Platform=$PLATFORM Configuration=$CONFIGURATION $UBT_ARCHFLAG "${TRAILINGARGS[@]}""
 
 # Add additional flags based on actions, arguments, and env properties
 AdditionalFlags=""
@@ -141,9 +141,9 @@ elif [ $ACTION == "metadata" ]; then
 	AdditionalFlags="-mode=ExportXcodeMetadata"
 fi
 
-echo Running dotnet Engine/Binaries/DotNET/UnrealBuildTool/UnrealBuildTool.dll $TARGET $PLATFORM $CONFIGURATION "$TRAILINGARGS" $UBT_ARCHFLAG $AdditionalFlags
+echo Running dotnet Engine/Binaries/DotNET/UnrealBuildTool/UnrealBuildTool.dll $TARGET $PLATFORM $CONFIGURATION "${TRAILINGARGS[@]}" $UBT_ARCHFLAG $AdditionalFlags
 # set an envvar to let UBT know that it's being run from xcode (envvar allows children to get the setting if needed)
-UE_BUILD_FROM_XCODE=1 dotnet Engine/Binaries/DotNET/UnrealBuildTool/UnrealBuildTool.dll $TARGET $PLATFORM $CONFIGURATION "$TRAILINGARGS" $UBT_ARCHFLAG $AdditionalFlags
+UE_BUILD_FROM_XCODE=1 dotnet Engine/Binaries/DotNET/UnrealBuildTool/UnrealBuildTool.dll $TARGET $PLATFORM $CONFIGURATION "${TRAILINGARGS[@]}" $UBT_ARCHFLAG $AdditionalFlags
 
 ExitCode=$?
 if [ $ExitCode -eq 254 ] || [ $ExitCode -eq 255 ] || [ $ExitCode -eq 2 ]; then
