Fix selected dotnet version on arm64 when GetDotNet.bat is run from an x64 program

#rnx

diff --git a/Engine/Build/BatchFiles/GetDotnetPath.bat b/Engine/Build/BatchFiles/GetDotnetPath.bat
--- a/Engine/Build/BatchFiles/GetDotnetPath.bat
+++ b/Engine/Build/BatchFiles/GetDotnetPath.bat
@@ -15,6 +15,10 @@ rem figure out which architecture to use
 set UE_DOTNET_ARCH=win-x64
 if "%PROCESSOR_ARCHITECTURE%" == "ARM64" (
 	set UE_DOTNET_ARCH=win-arm64
+) else if "%PROCESSOR_ARCHITECTURE%" == "AMD64" (
+	if "%ProgramFiles(Arm)%" NEQ "" (
+		set UE_DOTNET_ARCH=win-arm64
+	)
 )
 
 rem add the dotnet sdk in the repo as the current dotnet sdk
@@ -30,7 +34,7 @@ rem for dotnet applications that require older dotnet runtimes, this will allow
 set DOTNET_ROLL_FORWARD=LatestMajor
 
 rem skip dotnet verification when using our submitted sdk as we know it is up to date
-ECHO Using bundled DotNet SDK version: %UE_DOTNET_VERSION%
+ECHO Using bundled DotNet SDK version: %UE_DOTNET_VERSION% %UE_DOTNET_ARCH%
 goto Succeeded
 
 :verify_dotnet
