- Allow ResignApp.sh to just convert between say .xcarchive and .ipa

diff --git a/Engine/Build/BatchFiles/Mac/ResignApp.sh b/Engine/Build/BatchFiles/Mac/ResignApp.sh
--- a/Engine/Build/BatchFiles/Mac/ResignApp.sh
+++ b/Engine/Build/BatchFiles/Mac/ResignApp.sh
@@ -79,33 +79,35 @@ DoWork()
 
 	echo "------------------------------------------------------------------------------"
 
-	EXTRAOPTIONS=""
-	while IFS='' read -r line || [[ -n "$line" ]]; do
-
-		if [[ ! -z "$BUNDLE" ]] && [[ "$line" == *".appex"* ]]; then
-		   extbundleid=$(/usr/libexec/PlistBuddy -c "Print:CFBundleIdentifier" "$line/Info.plist")
-		   echo "Changing .appex BundleID $extbundleid with : ${extbundleid/$oldbundleid/$BUNDLE}"
-		   /usr/libexec/PlistBuddy -c "Set:CFBundleIdentifier ${extbundleid/$oldbundleid/$BUNDLE}" "$line/Info.plist"
-		fi    
-
-		if [[ -f "$line/embedded.mobileprovision" ]]; then
-			security cms -D -i "$line/embedded.mobileprovision" > t_entitlements_full.plist
-			/usr/libexec/PlistBuddy -x -c 'Print:Entitlements' t_entitlements_full.plist > t_entitlements.plist
-			TEAMID=$(/usr/libexec/PlistBuddy -c 'Print:com.apple.developer.team-identifier' t_entitlements.plist)
-			/usr/libexec/PlistBuddy -c "Set:application-identifier $TEAMID.$BUNDLE" t_entitlements.plist
-			
-			if [[ ! -z "$BUNDLE" ]]; then
-				EXTRAOPTIONS="--entitlements t_entitlements.plist"
-			else
-				EXTRAOPTIONS="--preserve-metadata=entitlements,flags,identifier"
+	if [[ ! -z "$DEVELOPER" ]]; then
+		EXTRAOPTIONS=""
+		while IFS='' read -r line || [[ -n "$line" ]]; do
+	
+			if [[ ! -z "$BUNDLE" ]] && [[ "$line" == *".appex"* ]]; then
+			   extbundleid=$(/usr/libexec/PlistBuddy -c "Print:CFBundleIdentifier" "$line/Info.plist")
+			   echo "Changing .appex BundleID $extbundleid with : ${extbundleid/$oldbundleid/$BUNDLE}"
+			   /usr/libexec/PlistBuddy -c "Set:CFBundleIdentifier ${extbundleid/$oldbundleid/$BUNDLE}" "$line/Info.plist"
+			fi    
+	
+			if [[ -f "$line/embedded.mobileprovision" ]]; then
+				security cms -D -i "$line/embedded.mobileprovision" > t_entitlements_full.plist
+				/usr/libexec/PlistBuddy -x -c 'Print:Entitlements' t_entitlements_full.plist > t_entitlements.plist
+				TEAMID=$(/usr/libexec/PlistBuddy -c 'Print:com.apple.developer.team-identifier' t_entitlements.plist)
+				/usr/libexec/PlistBuddy -c "Set:application-identifier $TEAMID.$BUNDLE" t_entitlements.plist
+				
+				if [[ ! -z "$BUNDLE" ]]; then
+					EXTRAOPTIONS="--entitlements t_entitlements.plist"
+				else
+					EXTRAOPTIONS="--preserve-metadata=entitlements,flags,identifier"
+				fi
 			fi
-		fi
-		
-		echo ""
-		echo Codesigning $line [/usr/bin/codesign  --continue -f -s "$DEVELOPER" $EXTRAOPTIONS "$line"]
-		/usr/bin/codesign --continue -f -s "$DEVELOPER" $EXTRAOPTIONS "$line"
-
-	done < directories.txt
+			
+			echo ""
+			echo Codesigning $line [/usr/bin/codesign  --continue -f -s "$DEVELOPER" $EXTRAOPTIONS "$line"]
+			/usr/bin/codesign --continue -f -s "$DEVELOPER" $EXTRAOPTIONS "$line"
+	
+		done < directories.txt
+	fi
 
 	if [[ "$TARGETAPP" == *.ipa ]]; then
 		echo ""
@@ -153,13 +155,17 @@ DoWork()
 Help()
 {
 	echo ""
-	echo "Usage: ResignApp <options>"
+	echo "Purpose: This tool can open an app, (optionally) resign it, and save it out as the same or dofferent format."
+	echo "Formats: .app, .ipa, .xcarchive"
+	echo "Note: Only IOS apps have been tested, Mac apps may not work. TVOS and VisionOS "
 	echo ""
 	echo "REQUIRED OPTIONS:"
 	echo "  -s | --sourceapp <path to input>"
 	echo "     Path to the .app or .ipa you want to resign"
 	echo "  -d | --destapp <path to output>"
 	echo "     Path to the output .app or .ipa (note that input type and output type do not need to match)"
+	echo ""
+	echo "OPTIONS REQUIRED FOR THE SIGNING OPTIONS:"
 	echo "  -i | --identity <signing identitiy> "
 	echo "     Name of the identity to sign with (this is the full name or prefix of a certificate in your login keychain - 'Apple Development' will often work)"
 	echo ""
