Fix -nocompileuat arg for RunUAT shell scripts

#rnx

diff --git a/Engine/Build/BatchFiles/RunUAT.sh b/Engine/Build/BatchFiles/RunUAT.sh
--- a/Engine/Build/BatchFiles/RunUAT.sh
+++ b/Engine/Build/BatchFiles/RunUAT.sh
@@ -58,8 +58,8 @@ if [ ! -f Build/BatchFiles/RunUAT.sh ]; then
 	exit 1
 fi
 
-# see if we have the no compile arg
-if echo "$@" | grep -q -w -i "\-nocompile"; then
+# see if we have the no compile uat arg
+if echo "$@" | grep -q -w -i "\-nocompileuat"; then
 	UATCompileArg=
 else
 	UATCompileArg=-compile
diff --git a/Engine/Build/BatchFiles/RunWineUAT.sh b/Engine/Build/BatchFiles/RunWineUAT.sh
--- a/Engine/Build/BatchFiles/RunWineUAT.sh
+++ b/Engine/Build/BatchFiles/RunWineUAT.sh
@@ -38,8 +38,8 @@ trap on_exit EXIT
 # Make WineUAT are built on Linux *not* Wine/Windows as it is very slow and broken
 # if UBT is not built, building WineUAT will build it
 
-# See if we have the no compile arg
-if echo "$@" | grep -q -w -i "\-nocompile"; then
+# See if we have the no compile uat arg
+if echo "$@" | grep -q -w -i "\-nocompileuat"; then
     WineUATCompileArg=
 else
     WineUATCompileArg=-compile
