Removing x11 server from script

#rnx

diff --git a/Engine/Build/BatchFiles/RunWineUAT.sh b/Engine/Build/BatchFiles/RunWineUAT.sh
--- a/Engine/Build/BatchFiles/RunWineUAT.sh
+++ b/Engine/Build/BatchFiles/RunWineUAT.sh
@@ -23,17 +23,6 @@ rm -rf ../qagame
 # Reformat UE_SDKS_ROOT under Z: and with backslashes
 export UE_SDKS_ROOT="Z:${UE_SDKS_ROOT//\//\\}"
 
-# Pretty hacky, depends on this package being around but wine/windows works so much better with an x11 server running
-# Avoids listening on TCP and Unix sockets as that yields errors when running as non-root
-Xvfb :99 -screen 0 1024x768x16 -nolisten tcp -nolisten unix &
-XVFB_PID=$!
-if ! (ps -p $XVFB_PID > /dev/null); then echo 'Failed to start xvfb process'; exit 1; fi
-
-function on_exit {
-    echo "Kill XVFB '$XVFB_PID'"
-    kill $XVFB_PID
-}
-trap on_exit EXIT
 
 # Make WineUAT are built on Linux *not* Wine/Windows as it is very slow and broken
 # if UBT is not built, building WineUAT will build it
