Change wine64 to wine

#rnx

diff --git a/Engine/Build/BatchFiles/RunWineUAT.sh b/Engine/Build/BatchFiles/RunWineUAT.sh
--- a/Engine/Build/BatchFiles/RunWineUAT.sh
+++ b/Engine/Build/BatchFiles/RunWineUAT.sh
@@ -62,4 +62,4 @@ if [ "$WineUATCompileArg" = "-compile" ]; then
 fi
 
 # Never UBT/UAT compile through Wine (see comment above)
-wine64 "$SCRIPT_DIR/RunUAT.bat" -nocompileuat -nocompile -NoP4 "$@" 2>&1
\ No newline at end of file
+wine "$SCRIPT_DIR/RunUAT.bat" -nocompileuat -nocompile -NoP4 "$@" 2>&1
