macOS quality of life. When downloading from GitHub, a flurry of quarantine messages come up when running Setup on macOS. Remove quarantine to allow running post Github fetch Setup out of the box.

diff --git a/Engine/Build/BatchFiles/Mac/GitDependencies.sh b/Engine/Build/BatchFiles/Mac/GitDependencies.sh
--- a/Engine/Build/BatchFiles/Mac/GitDependencies.sh
+++ b/Engine/Build/BatchFiles/Mac/GitDependencies.sh
@@ -16,7 +16,15 @@ cd ../../../..
 ARCH=x64
 [ $(uname -m) == "arm64" ] && ARCH=arm64 
 
-./Engine/Binaries/DotNET/GitDependencies/osx-$ARCH/GitDependencies "$@"
+GIT_DEP_BIN_DIR=./Engine/Binaries/DotNET/GitDependencies/osx-$ARCH
+GIT_DEP_EXE=$GIT_DEP_BIN_DIR/GitDependencies
+
+if [[ $(xattr $GIT_DEP_EXE) = "com.apple.quarantine" ]]
+then
+    xattr -d com.apple.quarantine $GIT_DEP_EXE $GIT_DEP_BIN_DIR/*.dylib
+fi
+
+$GIT_DEP_EXE "$@"
 
 pushd "$(dirname "$SCRIPT_PATH")" > /dev/null
 sh FixDependencyFiles.sh
