Disable dotnet msbuild node reuse in BuildUBT/UAT batch files

#rnx

diff --git a/Engine/Build/BatchFiles/BuildUAT.bat b/Engine/Build/BatchFiles/BuildUAT.bat
--- a/Engine/Build/BatchFiles/BuildUAT.bat
+++ b/Engine/Build/BatchFiles/BuildUAT.bat
@@ -40,7 +40,7 @@ if /I "%2" == "FORCE" del %DEPENDS_FILE% >nul 2>&1
 
 rem Check to see if the files in the AutomationTool solution have changed
 mkdir "%TEMP_DEPENDS_PATH%"
-dotnet msbuild Programs\AutomationTool\AutomationTool.sln -t:ProjectDependencies -p:Configuration=Development -p:Platform="Any CPU" -p:OutputPath="%TEMP_DEPENDS_PATH%" -noLogo -v:%MSBUILD_LOGLEVEL%
+dotnet msbuild Programs\AutomationTool\AutomationTool.sln -t:ProjectDependencies -p:Configuration=Development -p:Platform="Any CPU" -p:OutputPath="%TEMP_DEPENDS_PATH%" -noLogo -v:%MSBUILD_LOGLEVEL% -nr:false
 if errorlevel 1 goto Error_UATDependsFailed
 copy /b "%TEMP_DEPENDS_PATH%\*.dep.csv" "%TEMP_DEPENDS_FILE%" >nul 2>&1
 "%SYSTEMROOT%\System32\sort.exe" /unique /o "%TEMP_DEPENDS_FILE%" "%TEMP_DEPENDS_FILE%"
diff --git a/Engine/Build/BatchFiles/BuildUAT.sh b/Engine/Build/BatchFiles/BuildUAT.sh
--- a/Engine/Build/BatchFiles/BuildUAT.sh
+++ b/Engine/Build/BatchFiles/BuildUAT.sh
@@ -35,7 +35,7 @@ mkdir -p ../Intermediate/Build
 DEPENDS_FILE=../Intermediate/Build/AutomationTool.dep.csv
 TEMP_DEPENDS_FILE=$(mktemp)
 
-dotnet msbuild Programs/AutomationTool/AutomationTool.sln -t:ProjectDependencies -p:Configuration=Development -p:DependsPath="$TEMP_DEPENDS_FILE" -p:DependsOverwrite=false -m:1 -noLogo -v:"$MSBuild_Verbosity"
+dotnet msbuild Programs/AutomationTool/AutomationTool.sln -t:ProjectDependencies -p:Configuration=Development -p:DependsPath="$TEMP_DEPENDS_FILE" -p:DependsOverwrite=false -m:1 -noLogo -v:"$MSBuild_Verbosity" -nr:false
 sort -u -o "$TEMP_DEPENDS_FILE" "$TEMP_DEPENDS_FILE"
 
 PERFORM_REBUILD=0
diff --git a/Engine/Build/BatchFiles/BuildUBT.bat b/Engine/Build/BatchFiles/BuildUBT.bat
--- a/Engine/Build/BatchFiles/BuildUBT.bat
+++ b/Engine/Build/BatchFiles/BuildUBT.bat
@@ -36,7 +36,7 @@ if /I "%2" == "FORCE" del %DEPENDS_FILE% >nul 2>&1
 
 rem Check to see if the files in the UnrealBuildTool solution have changed
 mkdir "%TEMP_DEPENDS_PATH%"
-dotnet msbuild Programs\UnrealBuildTool\UnrealBuildTool.sln -t:ProjectDependencies -p:Configuration=Development -p:Platform="Any CPU" -p:OutputPath="%TEMP_DEPENDS_PATH%" -noLogo -v:%MSBUILD_LOGLEVEL%
+dotnet msbuild Programs\UnrealBuildTool\UnrealBuildTool.sln -t:ProjectDependencies -p:Configuration=Development -p:Platform="Any CPU" -p:OutputPath="%TEMP_DEPENDS_PATH%" -noLogo -v:%MSBUILD_LOGLEVEL% -nr:false
 if errorlevel 1 goto Error_UBTDependsFailed
 copy /b "%TEMP_DEPENDS_PATH%\*.dep.csv" "%TEMP_DEPENDS_FILE%" >nul 2>&1
 "%SYSTEMROOT%\System32\sort.exe" /unique /o "%TEMP_DEPENDS_FILE%" "%TEMP_DEPENDS_FILE%"
diff --git a/Engine/Build/BatchFiles/BuildUBT.sh b/Engine/Build/BatchFiles/BuildUBT.sh
--- a/Engine/Build/BatchFiles/BuildUBT.sh
+++ b/Engine/Build/BatchFiles/BuildUBT.sh
@@ -36,7 +36,7 @@ mkdir -p ../Intermediate/Build
 DEPENDS_FILE=../Intermediate/Build/UnrealBuildTool.dep.csv
 TEMP_DEPENDS_FILE=$(mktemp)
 
-dotnet msbuild Programs/UnrealBuildTool/UnrealBuildTool.sln -t:ProjectDependencies -p:Configuration=Development -p:DependsPath="$TEMP_DEPENDS_FILE" -p:DependsOverwrite=false -m:1 -noLogo -v:"$MSBuild_Verbosity"
+dotnet msbuild Programs/UnrealBuildTool/UnrealBuildTool.sln -t:ProjectDependencies -p:Configuration=Development -p:DependsPath="$TEMP_DEPENDS_FILE" -p:DependsOverwrite=false -m:1 -noLogo -v:"$MSBuild_Verbosity" -nr:false
 sort -u -o "$TEMP_DEPENDS_FILE" "$TEMP_DEPENDS_FILE"
 
 PERFORM_REBUILD=0
