Don't rerun SetupDotnet.sh if not necessary

#rnx

diff --git a/Engine/Build/BatchFiles/Linux/SetupDotnet.sh b/Engine/Build/BatchFiles/Linux/SetupDotnet.sh
--- a/Engine/Build/BatchFiles/Linux/SetupDotnet.sh
+++ b/Engine/Build/BatchFiles/Linux/SetupDotnet.sh
@@ -1,6 +1,11 @@
 #!/bin/bash
 # Copyright Epic Games, Inc. All Rights Reserved.
 
+# Check if already set up
+if [ ! "$UE_DOTNET_DIR" == "" ]; then
+	return
+fi
+
 START_DIR=`pwd`
 cd "$1"
 
diff --git a/Engine/Build/BatchFiles/Mac/SetupDotnet.sh b/Engine/Build/BatchFiles/Mac/SetupDotnet.sh
--- a/Engine/Build/BatchFiles/Mac/SetupDotnet.sh
+++ b/Engine/Build/BatchFiles/Mac/SetupDotnet.sh
@@ -1,6 +1,11 @@
 #!/bin/bash
 # Copyright Epic Games, Inc. All Rights Reserved.
 
+# Check if already set up
+if [ ! "$UE_DOTNET_DIR" == "" ]; then
+	return
+fi
+
 START_DIR=`pwd`
 cd "$1"
 
