Refactor dotnet dependency scan into batch file

#rnx

diff --git a/Engine/Build/BatchFiles/BuildUAT.bat b/Engine/Build/BatchFiles/BuildUAT.bat
--- a/Engine/Build/BatchFiles/BuildUAT.bat
+++ b/Engine/Build/BatchFiles/BuildUAT.bat
@@ -26,10 +26,7 @@ if errorlevel 1 goto Error_UBTCompileFailed
 
 rem ## Properties
 set DEPENDS_FILE=..\Intermediate\Build\AutomationTool.dep.csv 
-:tempRetry
-set TEMP_DEPENDS_PATH=%TEMP%\BuildUAT-%random%%random%
-if exist %TEMP_DEPENDS_PATH% goto :tempRetry
-set TEMP_DEPENDS_FILE=%TEMP_DEPENDS_PATH%.dep.csv
+set TEMP_DEPENDS_FILE=..\Intermediate\Build\AutomationTool.dep.csv.tmp
 
 rem ## Command line arguments
 set MSBUILD_LOGLEVEL=%1
@@ -38,12 +35,8 @@ if not defined %MSBUILD_LOGLEVEL set MSBUILD_LOGLEVEL=quiet
 rem Force build by deleting dependency file
 if /I "%2" == "FORCE" del %DEPENDS_FILE% >nul 2>&1
 
-rem Check to see if the files in the AutomationTool solution have changed
-mkdir "%TEMP_DEPENDS_PATH%"
-dotnet msbuild Programs\AutomationTool\AutomationTool.sln -t:ProjectDependencies -p:Configuration=Development -p:Platform="Any CPU" -p:OutputPath="%TEMP_DEPENDS_PATH%" -noLogo -v:%MSBUILD_LOGLEVEL%
-if errorlevel 1 goto Error_UATDependsFailed
-copy /b "%TEMP_DEPENDS_PATH%\*.dep.csv" "%TEMP_DEPENDS_FILE%" >nul 2>&1
-"%SYSTEMROOT%\System32\sort.exe" /unique /o "%TEMP_DEPENDS_FILE%" "%TEMP_DEPENDS_FILE%"
+rem Check to see if the files in the UnrealBuildTool solution have changed
+call "%~dp0DotnetDepends.bat" Programs\AutomationTool\AutomationTool.sln %TEMP_DEPENDS_FILE% %MSBUILD_LOGLEVEL%
 if errorlevel 1 goto Error_UATDependsFailed
 
 :Check_UpToDate
@@ -106,7 +99,6 @@ goto Exit
 :Exit
 rem ## Clean up temp files
 del "%TEMP_DEPENDS_FILE%" >nul 2>&1
-rmdir /s /q "%TEMP_DEPENDS_PATH%" >nul 2>&1
 rem ## Restore original CWD in case we change it
 popd
 exit /B %RUNUAT_EXITCODE%
diff --git a/Engine/Build/BatchFiles/BuildUBT.bat b/Engine/Build/BatchFiles/BuildUBT.bat
--- a/Engine/Build/BatchFiles/BuildUBT.bat
+++ b/Engine/Build/BatchFiles/BuildUBT.bat
@@ -21,11 +21,8 @@ call "%~dp0GetDotnetPath.bat"
 if errorlevel 1 goto Error_NoDotnetSDK
 
 rem ## Properties
-set DEPENDS_FILE=..\Intermediate\Build\UnrealBuildTool.dep.csv 
-:tempRetry
-set TEMP_DEPENDS_PATH=%TEMP%\BuildUBT-%random%%random%
-if exist %TEMP_DEPENDS_PATH% goto :tempRetry
-set TEMP_DEPENDS_FILE=%TEMP_DEPENDS_PATH%.dep.csv
+set DEPENDS_FILE=..\Intermediate\Build\UnrealBuildTool.dep.csv
+set TEMP_DEPENDS_FILE=..\Intermediate\Build\UnrealBuildTool.dep.csv.tmp
 
 rem ## Command line arguments
 set MSBUILD_LOGLEVEL=%1
@@ -35,11 +32,7 @@ rem Force build by deleting dependency file
 if /I "%2" == "FORCE" del %DEPENDS_FILE% >nul 2>&1
 
 rem Check to see if the files in the UnrealBuildTool solution have changed
-mkdir "%TEMP_DEPENDS_PATH%"
-dotnet msbuild Programs\UnrealBuildTool\UnrealBuildTool.sln -t:ProjectDependencies -p:Configuration=Development -p:Platform="Any CPU" -p:OutputPath="%TEMP_DEPENDS_PATH%" -noLogo -v:%MSBUILD_LOGLEVEL%
-if errorlevel 1 goto Error_UBTDependsFailed
-copy /b "%TEMP_DEPENDS_PATH%\*.dep.csv" "%TEMP_DEPENDS_FILE%" >nul 2>&1
-"%SYSTEMROOT%\System32\sort.exe" /unique /o "%TEMP_DEPENDS_FILE%" "%TEMP_DEPENDS_FILE%"
+call "%~dp0DotnetDepends.bat" Programs\UnrealBuildTool\UnrealBuildTool.sln %TEMP_DEPENDS_FILE% %MSBUILD_LOGLEVEL%
 if errorlevel 1 goto Error_UBTDependsFailed
 
 :Check_UpToDate
@@ -98,7 +91,6 @@ goto Exit
 :Exit
 rem ## Clean up temp files
 del "%TEMP_DEPENDS_FILE%" >nul 2>&1
-rmdir /s /q "%TEMP_DEPENDS_PATH%" >nul 2>&1
 rem ## Restore original CWD in case we change it
 popd
 exit /B %RUNUBT_EXITCODE%
diff --git a/Engine/Build/BatchFiles/DotnetDepends.bat b/Engine/Build/BatchFiles/DotnetDepends.bat
new file mode 100755
--- /dev/null
+++ b/Engine/Build/BatchFiles/DotnetDepends.bat
@@ -0,0 +1,51 @@
+@echo off
+
+rem ## Unreal Engine UnrealBuildTool build script
+rem ## Copyright Epic Games, Inc. All Rights Reserved.
+rem ##
+rem ## This script is expecting to exist in the Engine/Build/BatchFiles directory.  It will not work correctly
+rem ## if you copy it to a different location and run it.
+
+setlocal
+
+rem ## First, make sure the batch file exists in the folder we expect it to.  This is necessary in order to
+rem ## verify that our relative path to the /Engine/Source directory is correct
+if not exist "%~dp0..\..\Source" goto Error_BatchFileInWrongLocation
+
+rem ## Change the CWD to /Engine/Source.
+pushd "%~dp0..\..\Source"
+if not exist ..\Build\BatchFiles\DotnetDepends.bat goto Error_DependsFailed
+
+rem ## Verify that dotnet is present
+call "%~dp0GetDotnetPath.bat"
+if errorlevel 1 goto Error_DependsFailed
+
+set SLN=%1
+set DEPENDS_FILE=%2
+set MSBUILD_LOGLEVEL=%3
+if not defined %MSBUILD_LOGLEVEL set MSBUILD_LOGLEVEL=quiet
+
+:tempRetry
+set TEMP_DEPENDS_PATH=%TEMP%\DotnetDepends-%random%%random%
+if exist %TEMP_DEPENDS_PATH% goto :tempRetry
+set TEMP_DEPENDS_FILE=%TEMP_DEPENDS_PATH%\out.csv
+
+rem Check to see if the files in the UnrealBuildTool solution have changed
+mkdir "%TEMP_DEPENDS_PATH%"
+dotnet msbuild "%SLN%" -t:Scan -p:Configuration=Development -p:Platform="Any CPU" -p:OutputPath="%TEMP_DEPENDS_PATH%\\" -p:DependsEncoding=Ascii -noLogo -v:%MSBUILD_LOGLEVEL%
+if errorlevel 1 goto Error_DependsFailed
+type "%TEMP_DEPENDS_PATH%\*.dep.csv" > "%TEMP_DEPENDS_FILE%" 2>nul
+"%SYSTEMROOT%\System32\sort.exe" /+64 /unique /o "%DEPENDS_FILE%" "%TEMP_DEPENDS_FILE%"
+if errorlevel 1 goto Error_DependsFailed
+
+set DEPENDS_EXITCODE=0
+goto Exit
+
+:Error_DependsFailed
+set DEPENDS_EXITCODE=1
+goto Exit
+
+:Exit
+rem ## Clean up temp files
+rmdir /s /q "%TEMP_DEPENDS_PATH%" >nul 2>&1
+exit /B %DEPENDS_EXITCODE%
\ No newline at end of file
