Update linux build scripts to check dotnet dependencies in parallel

#rnx

diff --git a/Engine/Build/BatchFiles/BuildUAT.sh b/Engine/Build/BatchFiles/BuildUAT.sh
--- a/Engine/Build/BatchFiles/BuildUAT.sh
+++ b/Engine/Build/BatchFiles/BuildUAT.sh
@@ -35,8 +35,7 @@ mkdir -p ../Intermediate/Build
 DEPENDS_FILE=../Intermediate/Build/AutomationTool.dep.csv
 TEMP_DEPENDS_FILE=$(mktemp)
 
-dotnet msbuild Programs/AutomationTool/AutomationTool.sln -t:ProjectDependencies -p:Configuration=Development -p:DependsPath="$TEMP_DEPENDS_FILE" -p:DependsOverwrite=false -m:1 -noLogo -v:"$MSBuild_Verbosity"
-sort -u -o "$TEMP_DEPENDS_FILE" "$TEMP_DEPENDS_FILE"
+"$SCRIPT_DIR/DotnetDepends.sh" Programs/AutomationTool/AutomationTool.sln "$TEMP_DEPENDS_FILE" "$MSBuild_Verbosity"
 
 PERFORM_REBUILD=0
 if [ "$2" == "FORCE" ]; then
diff --git a/Engine/Build/BatchFiles/BuildUBT.sh b/Engine/Build/BatchFiles/BuildUBT.sh
--- a/Engine/Build/BatchFiles/BuildUBT.sh
+++ b/Engine/Build/BatchFiles/BuildUBT.sh
@@ -36,8 +36,7 @@ mkdir -p ../Intermediate/Build
 DEPENDS_FILE=../Intermediate/Build/UnrealBuildTool.dep.csv
 TEMP_DEPENDS_FILE=$(mktemp)
 
-dotnet msbuild Programs/UnrealBuildTool/UnrealBuildTool.sln -t:ProjectDependencies -p:Configuration=Development -p:DependsPath="$TEMP_DEPENDS_FILE" -p:DependsOverwrite=false -m:1 -noLogo -v:"$MSBuild_Verbosity"
-sort -u -o "$TEMP_DEPENDS_FILE" "$TEMP_DEPENDS_FILE"
+"$SCRIPT_DIR/DotnetDepends.sh" Programs/UnrealBuildTool/UnrealBuildTool.sln "$TEMP_DEPENDS_FILE" "$MSBuild_Verbosity"
 
 PERFORM_REBUILD=0
 if [ "$2" == "FORCE" ]; then
diff --git a/Engine/Build/BatchFiles/DotnetDepends.sh b/Engine/Build/BatchFiles/DotnetDepends.sh
new file mode 100755
--- /dev/null
+++ b/Engine/Build/BatchFiles/DotnetDepends.sh
@@ -0,0 +1,39 @@
+#!/bin/bash
+
+## Unreal Engine AutomationTool build script
+## Copyright Epic Games, Inc. All Rights Reserved
+##
+## This script is expecting to exist in the Engine/Build/BatchFiles directory.  It will not work correctly
+## if you copy it to a different location and run it.
+
+## First, make sure the batch file exists in the folder we expect it to.  This is necessary in order to
+## verify that our relative path to the /Engine/Source directory is correct
+SCRIPT_DIR=$(cd "$(dirname "$0")" && pwd)
+pushd "$SCRIPT_DIR/../../Source" >/dev/null || exit 1
+
+if [ ! -f ../Build/BatchFiles/DotnetDepends.sh ]; then
+  echo
+  echo "DotnetDepends ERROR: The script does not appear to be located in the /Engine/Build/BatchFiles directory.  This script must be run from within that directory."
+  echo
+  popd >/dev/null || exit 1
+  exit 1
+fi
+
+if [ "$(uname)" = "Darwin" ]; then
+  # Setup Environment
+  source "$SCRIPT_DIR/Mac/SetupEnvironment.sh" -dotnet "$SCRIPT_DIR/Mac"
+elif [ "$(uname)" = "Linux" ]; then
+  # Setup Environment
+  source "$SCRIPT_DIR/Linux/SetupEnvironment.sh" -dotnet "$SCRIPT_DIR/Linux"
+fi
+
+SLN=$1
+DEPENDS_FILE=$2
+MSBuild_Verbosity="${3:-quiet}"
+TEMPDIR=$(mktemp -q -d)
+
+dotnet msbuild "$SLN" -t:Scan -p:Configuration=Development -p:Platform="Any CPU" -p:OutputPath="$TEMPDIR"/ -p:DependsEncoding=Ascii -noLogo -v:"$MSBuild_Verbosity"
+cat "$TEMPDIR"/*.dep.csv | sort -u -o "$DEPENDS_FILE"
+rm -rf "$TEMPDIR"
+
+popd >/dev/null || return
