[Backout] - CL39893233
#rnx
Original CL Desc
-----------------------------------------------------------------
Refactor dotnet dependency scan into batch file

#rnx

diff --git a/Engine/Build/BatchFiles/BuildUAT.bat b/Engine/Build/BatchFiles/BuildUAT.bat
--- a/Engine/Build/BatchFiles/BuildUAT.bat
+++ b/Engine/Build/BatchFiles/BuildUAT.bat
@@ -26,7 +26,10 @@ if errorlevel 1 goto Error_UBTCompileFailed
 
 rem ## Properties
 set DEPENDS_FILE=..\Intermediate\Build\AutomationTool.dep.csv 
-set TEMP_DEPENDS_FILE=..\Intermediate\Build\AutomationTool.dep.csv.tmp
+:tempRetry
+set TEMP_DEPENDS_PATH=%TEMP%\BuildUAT-%random%%random%
+if exist %TEMP_DEPENDS_PATH% goto :tempRetry
+set TEMP_DEPENDS_FILE=%TEMP_DEPENDS_PATH%.dep.csv
 
 rem ## Command line arguments
 set MSBUILD_LOGLEVEL=%1
@@ -35,8 +38,12 @@ if not defined %MSBUILD_LOGLEVEL set MSBUILD_LOGLEVEL=quiet
 rem Force build by deleting dependency file
 if /I "%2" == "FORCE" del %DEPENDS_FILE% >nul 2>&1
 
-rem Check to see if the files in the UnrealBuildTool solution have changed
-call "%~dp0DotnetDepends.bat" Programs\AutomationTool\AutomationTool.sln %TEMP_DEPENDS_FILE% %MSBUILD_LOGLEVEL%
+rem Check to see if the files in the AutomationTool solution have changed
+mkdir "%TEMP_DEPENDS_PATH%"
+dotnet msbuild Programs\AutomationTool\AutomationTool.sln -t:ProjectDependencies -p:Configuration=Development -p:Platform="Any CPU" -p:OutputPath="%TEMP_DEPENDS_PATH%" -noLogo -v:%MSBUILD_LOGLEVEL%
+if errorlevel 1 goto Error_UATDependsFailed
+copy /b "%TEMP_DEPENDS_PATH%\*.dep.csv" "%TEMP_DEPENDS_FILE%" >nul 2>&1
+"%SYSTEMROOT%\System32\sort.exe" /unique /o "%TEMP_DEPENDS_FILE%" "%TEMP_DEPENDS_FILE%"
 if errorlevel 1 goto Error_UATDependsFailed
 
 :Check_UpToDate
@@ -99,6 +106,7 @@ goto Exit
 :Exit
 rem ## Clean up temp files
 del "%TEMP_DEPENDS_FILE%" >nul 2>&1
+rmdir /s /q "%TEMP_DEPENDS_PATH%" >nul 2>&1
 rem ## Restore original CWD in case we change it
 popd
 exit /B %RUNUAT_EXITCODE%
diff --git a/Engine/Build/BatchFiles/BuildUBT.bat b/Engine/Build/BatchFiles/BuildUBT.bat
--- a/Engine/Build/BatchFiles/BuildUBT.bat
+++ b/Engine/Build/BatchFiles/BuildUBT.bat
@@ -21,8 +21,11 @@ call "%~dp0GetDotnetPath.bat"
 if errorlevel 1 goto Error_NoDotnetSDK
 
 rem ## Properties
-set DEPENDS_FILE=..\Intermediate\Build\UnrealBuildTool.dep.csv
-set TEMP_DEPENDS_FILE=..\Intermediate\Build\UnrealBuildTool.dep.csv.tmp
+set DEPENDS_FILE=..\Intermediate\Build\UnrealBuildTool.dep.csv 
+:tempRetry
+set TEMP_DEPENDS_PATH=%TEMP%\BuildUBT-%random%%random%
+if exist %TEMP_DEPENDS_PATH% goto :tempRetry
+set TEMP_DEPENDS_FILE=%TEMP_DEPENDS_PATH%.dep.csv
 
 rem ## Command line arguments
 set MSBUILD_LOGLEVEL=%1
@@ -32,7 +35,11 @@ rem Force build by deleting dependency file
 if /I "%2" == "FORCE" del %DEPENDS_FILE% >nul 2>&1
 
 rem Check to see if the files in the UnrealBuildTool solution have changed
-call "%~dp0DotnetDepends.bat" Programs\UnrealBuildTool\UnrealBuildTool.sln %TEMP_DEPENDS_FILE% %MSBUILD_LOGLEVEL%
+mkdir "%TEMP_DEPENDS_PATH%"
+dotnet msbuild Programs\UnrealBuildTool\UnrealBuildTool.sln -t:ProjectDependencies -p:Configuration=Development -p:Platform="Any CPU" -p:OutputPath="%TEMP_DEPENDS_PATH%" -noLogo -v:%MSBUILD_LOGLEVEL%
+if errorlevel 1 goto Error_UBTDependsFailed
+copy /b "%TEMP_DEPENDS_PATH%\*.dep.csv" "%TEMP_DEPENDS_FILE%" >nul 2>&1
+"%SYSTEMROOT%\System32\sort.exe" /unique /o "%TEMP_DEPENDS_FILE%" "%TEMP_DEPENDS_FILE%"
 if errorlevel 1 goto Error_UBTDependsFailed
 
 :Check_UpToDate
@@ -91,6 +98,7 @@ goto Exit
 :Exit
 rem ## Clean up temp files
 del "%TEMP_DEPENDS_FILE%" >nul 2>&1
+rmdir /s /q "%TEMP_DEPENDS_PATH%" >nul 2>&1
 rem ## Restore original CWD in case we change it
 popd
 exit /B %RUNUBT_EXITCODE%
diff --git a/Engine/Build/BatchFiles/DotnetDepends.bat b/Engine/Build/BatchFiles/DotnetDepends.bat
deleted file mode 100755
