[Backout] - CL39894051
#rnx
Original CL Desc
-----------------------------------------------------------------
Update linux build scripts to check dotnet dependencies in parallel

#rnx

diff --git a/Engine/Build/BatchFiles/BuildUAT.sh b/Engine/Build/BatchFiles/BuildUAT.sh
--- a/Engine/Build/BatchFiles/BuildUAT.sh
+++ b/Engine/Build/BatchFiles/BuildUAT.sh
@@ -35,7 +35,8 @@ mkdir -p ../Intermediate/Build
 DEPENDS_FILE=../Intermediate/Build/AutomationTool.dep.csv
 TEMP_DEPENDS_FILE=$(mktemp)
 
-"$SCRIPT_DIR/DotnetDepends.sh" Programs/AutomationTool/AutomationTool.sln "$TEMP_DEPENDS_FILE" "$MSBuild_Verbosity"
+dotnet msbuild Programs/AutomationTool/AutomationTool.sln -t:ProjectDependencies -p:Configuration=Development -p:DependsPath="$TEMP_DEPENDS_FILE" -p:DependsOverwrite=false -m:1 -noLogo -v:"$MSBuild_Verbosity"
+sort -u -o "$TEMP_DEPENDS_FILE" "$TEMP_DEPENDS_FILE"
 
 PERFORM_REBUILD=0
 if [ "$2" == "FORCE" ]; then
diff --git a/Engine/Build/BatchFiles/BuildUBT.sh b/Engine/Build/BatchFiles/BuildUBT.sh
--- a/Engine/Build/BatchFiles/BuildUBT.sh
+++ b/Engine/Build/BatchFiles/BuildUBT.sh
@@ -36,7 +36,8 @@ mkdir -p ../Intermediate/Build
 DEPENDS_FILE=../Intermediate/Build/UnrealBuildTool.dep.csv
 TEMP_DEPENDS_FILE=$(mktemp)
 
-"$SCRIPT_DIR/DotnetDepends.sh" Programs/UnrealBuildTool/UnrealBuildTool.sln "$TEMP_DEPENDS_FILE" "$MSBuild_Verbosity"
+dotnet msbuild Programs/UnrealBuildTool/UnrealBuildTool.sln -t:ProjectDependencies -p:Configuration=Development -p:DependsPath="$TEMP_DEPENDS_FILE" -p:DependsOverwrite=false -m:1 -noLogo -v:"$MSBuild_Verbosity"
+sort -u -o "$TEMP_DEPENDS_FILE" "$TEMP_DEPENDS_FILE"
 
 PERFORM_REBUILD=0
 if [ "$2" == "FORCE" ]; then
diff --git a/Engine/Build/BatchFiles/DotnetDepends.sh b/Engine/Build/BatchFiles/DotnetDepends.sh
deleted file mode 100755
