Change WINEPREFIX

#rnx

diff --git a/Engine/Build/BatchFiles/RunWineUAT.sh b/Engine/Build/BatchFiles/RunWineUAT.sh
--- a/Engine/Build/BatchFiles/RunWineUAT.sh
+++ b/Engine/Build/BatchFiles/RunWineUAT.sh
@@ -10,7 +10,7 @@ set -o pipefail
 SCRIPT_DIR=$(cd -- "$(dirname -- "${BASH_SOURCE[0]}")" &> /dev/null && pwd)
 cd "$SCRIPT_DIR/../.." || exit
 
-export WINEPREFIX="$HOME/.wine"
+export WINEPREFIX=/opt/horde/wine-data
 export WINEDEBUG=-all
 export WINEARCH=win64
 export DISPLAY=:99
