Remove wine related env variables for testing.

#rnx

diff --git a/Engine/Build/BatchFiles/RunWineUAT.sh b/Engine/Build/BatchFiles/RunWineUAT.sh
--- a/Engine/Build/BatchFiles/RunWineUAT.sh
+++ b/Engine/Build/BatchFiles/RunWineUAT.sh
@@ -10,11 +10,6 @@ set -o pipefail
 SCRIPT_DIR=$(cd -- "$(dirname -- "${BASH_SOURCE[0]}")" &> /dev/null && pwd)
 cd "$SCRIPT_DIR/../.." || exit
 
-export WINEPREFIX=/opt/horde/wine-data
-export WINEDEBUG=-all
-export WINEARCH=win64
-export DISPLAY=:99
-
 # Remove empty "qagame" directory. This collides with directory "QAGame" and confuses directory caching in UAT/UBT.
 # This is a problem on Linux with ext4 and the Perforce history for this directory.
 rm -rf ../qagame
