Remove UAT and UBT workaround for Wine

#rnx

diff --git a/Engine/Build/BatchFiles/RunWineUAT.sh b/Engine/Build/BatchFiles/RunWineUAT.sh
--- a/Engine/Build/BatchFiles/RunWineUAT.sh
+++ b/Engine/Build/BatchFiles/RunWineUAT.sh
@@ -18,43 +18,5 @@ rm -rf ../qagame
 # Reformat UE_SDKS_ROOT under Z: and with backslashes
 export UE_SDKS_ROOT="Z:${UE_SDKS_ROOT//\//\\}"
 
-
-# Make WineUAT are built on Linux *not* Wine/Windows as it is very slow and broken
-# if UBT is not built, building WineUAT will build it
-
-# See if we have the no compile uat arg
-if echo "$@" | grep -q -w -i "\-nocompileuat"; then
-    WineUATCompileArg=
-else
-    WineUATCompileArg=-compile
-fi
-
-# Control toggling of msbuild verbosity for easier debugging
-if echo "$@" | grep -q -w -i "\-msbuild-verbose"; then
-    MSBuild_Verbosity=normal
-else
-    MSBuild_Verbosity=quiet
-fi
-
-if [ "$WineUATCompileArg" = "-compile" ]; then
-    # See if the .csproj exists to be compiled
-    if [ ! -f Source/Programs/AutomationTool/AutomationTool.csproj ]; then
-        echo No project to compile, attempting to use precompiled AutomationTool
-        WineUATCompileArg=
-    else
-        "$SCRIPT_DIR/BuildUAT.sh" $MSBuild_Verbosity
-        if [ $? -ne 0 ]; then
-            echo RunWineUAT ERROR: AutomationTool failed to compile.
-            exit 1
-        fi
-
-        # Run UAT once on Linux to initialize and compile scripts
-        # Currently, .NET builds don't compile under Wine
-        #
-        # Eat normal output, but still show errors
-        "$SCRIPT_DIR"/RunUAT.sh -help > /dev/null
-    fi
-fi
-
 # Never UBT/UAT compile through Wine (see comment above)
-/opt/uebuilder/bin/wine "$SCRIPT_DIR/RunUAT.bat" -nocompileuat -nocompile -NoP4 "$@" 2>&1
+/opt/uebuilder/bin/wine "$SCRIPT_DIR/RunUAT.bat" "$@" 2>&1
