Update linux toolchain scripts to clang 19.1.7 with matching libc++

- changed to use the in-tree rpmalloc for the windows toolchain
- build libc++ at the same time as we build our runtimes
- update configs to use an appropriate gcc version for the sysroot so that libstdc++.a doesn't end up with a glibc dependency that rocky 8 can't meet

diff --git a/Engine/Build/BatchFiles/Linux/Toolchain/RunMe.bat b/Engine/Build/BatchFiles/Linux/Toolchain/RunMe.bat
--- a/Engine/Build/BatchFiles/Linux/Toolchain/RunMe.bat
+++ b/Engine/Build/BatchFiles/Linux/Toolchain/RunMe.bat
@@ -1,9 +1,9 @@
 @REM @echo off
-
-set TOOLCHAIN_VERSION=v23
-set LLVM_VERSION=18.1.0
-set LLVM_BRANCH=release/18.x
-set LLVM_TAG=llvmorg-18.1.0
+set BUILD_ONLY=0
+set TOOLCHAIN_VERSION=v24
+set LLVM_VERSION=19.1.7
+set LLVM_BRANCH=release/19.x
+set LLVM_TAG=llvmorg-19.1.7
 
 set CMAKE_BINARY=%CD%\..\..\..\..\Extras\ThirdPartyNotUE\CMake\bin\cmake.exe
 set PYTHON_BINARY=%CD%\..\..\..\..\Binaries\ThirdParty\Python\Win64\python.exe
@@ -30,10 +30,12 @@ echo Using VisualStudio: %VS_VERSION%
 @REM We need to build in a directory with shorter path, so we avoid hitting path max limit.
 set ROOT_DIR=%CD%
 
-rmdir /S /Q %TEMP%\clang-build-%LLVM_VERSION% > nul
+echo %BUILD_ONLY%
+if %BUILD_ONLY%==0 rmdir /S /Q %TEMP%\clang-build-%LLVM_VERSION% > nul
 mkdir %TEMP%\clang-build-%LLVM_VERSION%
 pushd %TEMP%\clang-build-%LLVM_VERSION%
 
+if %BUILD_ONLY%==1 goto :build
 
 rem The MS provided tar can be used to create and extract zip files on windows
 mkdir OUTPUT
@@ -41,10 +43,6 @@ pushd OUTPUT
 tar xf %ROOT_DIR%\%FILENAME%-windows.zip
 popd
 
-echo Cloning rpmalloc to speed up clang/lld runtime operation
-git clone [REDACTED_PROFILE_URL]/rpmalloc -b main rpmalloc --depth 1
-
-
 echo Cloning LLVM (tag %LLVM_TAG% only)
 rem clone -b can also accept tag names
 git clone [REDACTED_PROFILE_URL]/llvm-project source -b %LLVM_TAG% --single-branch --depth 1 -c advice.detachedHead=false
@@ -56,17 +54,15 @@ echo Applying patches
 pushd "source"
 rem set DRY_RUN=--dry-run
 set DRY_RUN=
-rem disabled but left in place in case testing reveals that they are still needed
-rem "%PATCH_BINARY%" %DRY_RUN% -p 1 -i %ROOT_DIR%\patches\clang\default-dwarf-4.patch 
-rem "%PATCH_BINARY%" %DRY_RUN% -p 1 -i %ROOT_DIR%\patches\compiler-rt\manually-define-AT_HWCAP2.diff 
-rem "%PATCH_BINARY%" %DRY_RUN% -p 1 -i %ROOT_DIR%\patches\llvm\disable-auto-upgrade-debug-info.patch 
+rem libSupport currently has some incorrect quoting resulting in opt.exe not building properly on windows.
+"%PATCH_BINARY%" %DRY_RUN% -p0 -i %ROOT_DIR%\patches\llvm\libSupport.patch
 popd
 
-
+:build
 echo Building LLVM, clang, lld, bolt
 mkdir build_all
 pushd build_all
-%CMAKE_BINARY% -G %VS_VERSION% -DLLVM_ENABLE_PROJECTS=llvm;clang;lld -DLLVM_INTEGRATED_CRT_ALLOC=../rpmalloc -DCMAKE_MSVC_RUNTIME_LIBRARY=MultiThreaded -DCMAKE_BUILD_TYPE=Release -DCMAKE_C_FLAGS_RELEASE="/O2 /Ob3 /DNDEBUG /Zi /Gy" -DCMAKE_CXX_FLAGS_RELEASE="/O2 /Ob3 /DNDEBUG /Zi /Gy" -DCMAKE_EXE_LINKER_FLAGS_RELEASE="/DEBUG /INCREMENTAL:NO /OPT:REF /OPT:ICF" -DCMAKE_INSTALL_PREFIX="..\install" -DLLVM_INCLUDE_BENCHMARKS=OFF -DLLVM_TARGETS_TO_BUILD="AArch64;X86" -DLLVM_ENABLE_ZLIB=FORCE_ON -DZLIB_LIBRARY="%ZLIB_PATH%/lib/Win64/Release/zlibstatic.lib" -DZLIB_INCLUDE_DIR="%ZLIB_PATH%/include/" -DCLANG_REPOSITORY_STRING="github.com/llvm/llvm-project" "..\source\llvm"
+%CMAKE_BINARY% -G %VS_VERSION% -DLLVM_ENABLE_PROJECTS=llvm;clang;lld -DLLVM_ENABLE_RPMALLOC=ON -DCMAKE_BUILD_TYPE=Release -DCMAKE_C_FLAGS_RELEASE="/O2 /Ob3 /DNDEBUG /Zi /Gy" -DCMAKE_CXX_FLAGS_RELEASE="/O2 /Ob3 /DNDEBUG /Zi /Gy" -DCMAKE_EXE_LINKER_FLAGS_RELEASE="/DEBUG /INCREMENTAL:NO /OPT:REF /OPT:ICF" -DCMAKE_INSTALL_PREFIX="..\install" -DLLVM_INCLUDE_BENCHMARKS=OFF -DLLVM_TARGETS_TO_BUILD="AArch64;X86" -DLLVM_ENABLE_ZLIB=FORCE_ON -DZLIB_LIBRARY="%ZLIB_PATH%/lib/Win64/Release/zlibstatic.lib" -DZLIB_INCLUDE_DIR="%ZLIB_PATH%/include/" -DCLANG_REPOSITORY_STRING="github.com/llvm/llvm-project" "..\source\llvm"
 if ERRORLEVEL 1 goto endscript
 %CMAKE_BINARY% --build . --target install --config Release
 IF ERRORLEVEL 1 goto endscript
diff --git a/Engine/Build/BatchFiles/Linux/Toolchain/RunMe.sh b/Engine/Build/BatchFiles/Linux/Toolchain/RunMe.sh
--- a/Engine/Build/BatchFiles/Linux/Toolchain/RunMe.sh
+++ b/Engine/Build/BatchFiles/Linux/Toolchain/RunMe.sh
@@ -34,6 +34,9 @@ if ! [ -f "/.dockerenv" ]; then
   #echo Removing ${ImageName}...
   docker rm ${ImageName}
 
+  # Copy LibC++ files into place
+  update_libc++.sh
+
 else
 
 
@@ -75,7 +78,7 @@ else
 ##### Rocky Linux 8 stuff #####
 	dnf install -y 'dnf-command(config-manager)'
 	dnf config-manager --enable powertools
-	yum install -y ncurses-devel patch make cmake3 zip git wget which gcc-c++ bzip2 file unzip autoconf libtool glibc-devel libstdc++-devel xz flex bison file unzip autoconf libtool diffutils openssl openssl-devel rsync perl-libintl perl-Text-Unidecode mingw32-gcc mingw32-gcc-c++ mingw64-gcc mingw64-gcc-c++ glibc-static mingw32-winpthreads-static.noarch mingw64-winpthreads-static.noarch libstdc++-static texinfo help2man perl-Unicode-Normalize gettext python3
+	yum install -y ncurses-devel patch make cmake3 zip git wget which gcc-c++ bzip2 file unzip autoconf libtool glibc-devel libstdc++-devel xz flex bison file unzip autoconf libtool diffutils openssl openssl-devel rsync perl-libintl perl-Text-Unidecode mingw32-gcc mingw32-gcc-c++ mingw64-gcc mingw64-gcc-c++ glibc-static mingw32-winpthreads-static.noarch mingw64-winpthreads-static.noarch libstdc++-static texinfo help2man perl-Unicode-Normalize gettext python38
 
 ##### End Rocky Linux 8 stuff #####
   fi
diff --git a/Engine/Build/BatchFiles/Linux/Toolchain/aarch64-unknown-linux-gnueabi.linux.config b/Engine/Build/BatchFiles/Linux/Toolchain/aarch64-unknown-linux-gnueabi.linux.config
--- a/Engine/Build/BatchFiles/Linux/Toolchain/aarch64-unknown-linux-gnueabi.linux.config
+++ b/Engine/Build/BatchFiles/Linux/Toolchain/aarch64-unknown-linux-gnueabi.linux.config
@@ -15,6 +15,7 @@ CT_CONFIGURE_has_autoconf_2_65_or_newer=y
 CT_CONFIGURE_has_autoreconf_2_65_or_newer=y
 CT_CONFIGURE_has_automake_1_15_or_newer=y
 CT_CONFIGURE_has_gnu_m4_1_4_12_or_newer=y
+CT_CONFIGURE_has_python_3_4_or_newer=y
 CT_CONFIGURE_has_bison_2_7_or_newer=y
 CT_CONFIGURE_has_bison_3_0_4_or_newer=y
 CT_CONFIGURE_has_python=y
@@ -311,8 +312,6 @@ CT_LINUX_V_4_18=y
 # CT_LINUX_V_3_13 is not set
 # CT_LINUX_V_3_12 is not set
 # CT_LINUX_V_3_10 is not set
-# CT_LINUX_V_3_4 is not set
-# CT_LINUX_V_3_2 is not set
 CT_LINUX_VERSION="4.18.20"
 CT_LINUX_MIRRORS="$(CT_Mirrors kernel.org linux ${CT_LINUX_VERSION})"
 CT_LINUX_ARCHIVE_FILENAME="@{pkg_name}-@{version}"
@@ -331,6 +330,7 @@ CT_LINUX_later_than_4_8=y
 CT_LINUX_4_8_or_later=y
 CT_LINUX_later_than_3_7=y
 CT_LINUX_3_7_or_later=y
+CT_LINUX_REQUIRE_3_7_or_later=y
 CT_LINUX_later_than_3_2=y
 CT_LINUX_3_2_or_later=y
 CT_LINUX_REQUIRE_3_2_or_later=y
@@ -435,6 +435,16 @@ CT_GLIBC_PKG_NAME="glibc"
 CT_GLIBC_SRC_RELEASE=y
 # CT_GLIBC_SRC_DEVEL is not set
 CT_GLIBC_PATCH_ORDER="global"
+# CT_GLIBC_V_2_38 is not set
+# CT_GLIBC_V_2_37 is not set
+# CT_GLIBC_V_2_36 is not set
+# CT_GLIBC_V_2_35 is not set
+# CT_GLIBC_V_2_34 is not set
+# CT_GLIBC_V_2_33 is not set
+# CT_GLIBC_V_2_32 is not set
+# CT_GLIBC_V_2_31 is not set
+# CT_GLIBC_V_2_30 is not set
+# CT_GLIBC_V_2_29 is not set
 CT_GLIBC_V_2_28=y
 # CT_GLIBC_V_2_27 is not set
 # CT_GLIBC_V_2_26 is not set
@@ -465,7 +475,6 @@ CT_GLIBC_2_30_or_older=y
 CT_GLIBC_older_than_2_30=y
 CT_GLIBC_2_29_or_older=y
 CT_GLIBC_older_than_2_29=y
-CT_GLIBC_REQUIRE_older_than_2_29=y
 CT_GLIBC_2_28_or_later=y
 CT_GLIBC_2_28_or_older=y
 CT_GLIBC_later_than_2_27=y
@@ -488,10 +497,6 @@ CT_GLIBC_DEP_KERNEL_HEADERS_VERSION=y
 CT_GLIBC_DEP_BINUTILS=y
 CT_GLIBC_DEP_GCC=y
 CT_GLIBC_DEP_PYTHON=y
-
-#
-# glibc >=2.29 requires python >3.4
-#
 CT_GLIBC_SPARC_ALLOW_V7=y
 CT_THREADS="nptl"
 CT_GLIBC_BUILD_SSP=y
@@ -517,6 +522,7 @@ CT_GLIBC_SSP_DEFAULT=y
 # CT_GLIBC_SSP_YES is not set
 # CT_GLIBC_SSP_ALL is not set
 # CT_GLIBC_SSP_STRONG is not set
+CT_GLIBC_ENABLE_WERROR=y
 CT_GLIBC_ENABLE_COMMON_FLAG=y
 CT_ALL_LIBC_CHOICES="AVR_LIBC GLIBC MINGW_W64 MOXIEBOX MUSL NEWLIB NONE PICOLIBC UCLIBC_NG"
 CT_LIBC_SUPPORT_THREADS_ANY=y
@@ -559,32 +565,32 @@ CT_GCC_PKG_NAME="gcc"
 CT_GCC_SRC_RELEASE=y
 # CT_GCC_SRC_DEVEL is not set
 CT_GCC_PATCH_ORDER="global"
-CT_GCC_V_13=y
+# CT_GCC_V_13 is not set
 # CT_GCC_V_12 is not set
 # CT_GCC_V_11 is not set
 # CT_GCC_V_10 is not set
 # CT_GCC_V_9 is not set
-# CT_GCC_V_8 is not set
+CT_GCC_V_8=y
 # CT_GCC_V_7 is not set
 # CT_GCC_V_6 is not set
 # CT_GCC_V_5 is not set
 # CT_GCC_V_4_9 is not set
-CT_GCC_VERSION="13.2.0"
+CT_GCC_VERSION="8.5.0"
 CT_GCC_MIRRORS="$(CT_Mirrors GNU gcc/gcc-${CT_GCC_VERSION}) $(CT_Mirrors sourceware gcc/releases/gcc-${CT_GCC_VERSION})"
 CT_GCC_ARCHIVE_FILENAME="@{pkg_name}-@{version}"
 CT_GCC_ARCHIVE_DIRNAME="@{pkg_name}-@{version}"
 CT_GCC_ARCHIVE_FORMATS=".tar.xz .tar.gz"
 CT_GCC_SIGNATURE_FORMAT=""
-CT_GCC_later_than_13=y
-CT_GCC_13_or_later=y
-CT_GCC_later_than_12=y
-CT_GCC_12_or_later=y
-CT_GCC_later_than_11=y
-CT_GCC_11_or_later=y
-CT_GCC_later_than_10=y
-CT_GCC_10_or_later=y
-CT_GCC_later_than_9=y
-CT_GCC_9_or_later=y
+CT_GCC_13_or_older=y
+CT_GCC_older_than_13=y
+CT_GCC_12_or_older=y
+CT_GCC_older_than_12=y
+CT_GCC_11_or_older=y
+CT_GCC_older_than_11=y
+CT_GCC_10_or_older=y
+CT_GCC_older_than_10=y
+CT_GCC_9_or_older=y
+CT_GCC_older_than_9=y
 CT_GCC_later_than_8=y
 CT_GCC_8_or_later=y
 CT_GCC_later_than_7=y
@@ -625,7 +631,6 @@ CT_CC_GCC_LIBSTDCXX_VERBOSE=m
 # Misc. obscure options.
 #
 CT_CC_CXA_ATEXIT=y
-CT_CC_GCC_TM_CLONE_REGISTRY=m
 # CT_CC_GCC_DISABLE_PCH is not set
 CT_CC_GCC_SJLJ_EXCEPTIONS=m
 CT_CC_GCC_LDBL_128=m
diff --git a/Engine/Build/BatchFiles/Linux/Toolchain/aarch64-unknown-linux-gnueabi.windows.config b/Engine/Build/BatchFiles/Linux/Toolchain/aarch64-unknown-linux-gnueabi.windows.config
--- a/Engine/Build/BatchFiles/Linux/Toolchain/aarch64-unknown-linux-gnueabi.windows.config
+++ b/Engine/Build/BatchFiles/Linux/Toolchain/aarch64-unknown-linux-gnueabi.windows.config
@@ -15,6 +15,7 @@ CT_CONFIGURE_has_autoconf_2_65_or_newer=y
 CT_CONFIGURE_has_autoreconf_2_65_or_newer=y
 CT_CONFIGURE_has_automake_1_15_or_newer=y
 CT_CONFIGURE_has_gnu_m4_1_4_12_or_newer=y
+CT_CONFIGURE_has_python_3_4_or_newer=y
 CT_CONFIGURE_has_bison_2_7_or_newer=y
 CT_CONFIGURE_has_bison_3_0_4_or_newer=y
 CT_CONFIGURE_has_python=y
@@ -316,8 +317,6 @@ CT_LINUX_V_4_18=y
 # CT_LINUX_V_3_13 is not set
 # CT_LINUX_V_3_12 is not set
 # CT_LINUX_V_3_10 is not set
-# CT_LINUX_V_3_4 is not set
-# CT_LINUX_V_3_2 is not set
 CT_LINUX_VERSION="4.18.20"
 CT_LINUX_MIRRORS="$(CT_Mirrors kernel.org linux ${CT_LINUX_VERSION})"
 CT_LINUX_ARCHIVE_FILENAME="@{pkg_name}-@{version}"
@@ -336,6 +335,7 @@ CT_LINUX_later_than_4_8=y
 CT_LINUX_4_8_or_later=y
 CT_LINUX_later_than_3_7=y
 CT_LINUX_3_7_or_later=y
+CT_LINUX_REQUIRE_3_7_or_later=y
 CT_LINUX_later_than_3_2=y
 CT_LINUX_3_2_or_later=y
 CT_LINUX_REQUIRE_3_2_or_later=y
@@ -440,6 +440,16 @@ CT_GLIBC_PKG_NAME="glibc"
 CT_GLIBC_SRC_RELEASE=y
 # CT_GLIBC_SRC_DEVEL is not set
 CT_GLIBC_PATCH_ORDER="global"
+# CT_GLIBC_V_2_38 is not set
+# CT_GLIBC_V_2_37 is not set
+# CT_GLIBC_V_2_36 is not set
+# CT_GLIBC_V_2_35 is not set
+# CT_GLIBC_V_2_34 is not set
+# CT_GLIBC_V_2_33 is not set
+# CT_GLIBC_V_2_32 is not set
+# CT_GLIBC_V_2_31 is not set
+# CT_GLIBC_V_2_30 is not set
+# CT_GLIBC_V_2_29 is not set
 CT_GLIBC_V_2_28=y
 # CT_GLIBC_V_2_27 is not set
 # CT_GLIBC_V_2_26 is not set
@@ -470,7 +480,6 @@ CT_GLIBC_2_30_or_older=y
 CT_GLIBC_older_than_2_30=y
 CT_GLIBC_2_29_or_older=y
 CT_GLIBC_older_than_2_29=y
-CT_GLIBC_REQUIRE_older_than_2_29=y
 CT_GLIBC_2_28_or_later=y
 CT_GLIBC_2_28_or_older=y
 CT_GLIBC_later_than_2_27=y
@@ -493,10 +502,6 @@ CT_GLIBC_DEP_KERNEL_HEADERS_VERSION=y
 CT_GLIBC_DEP_BINUTILS=y
 CT_GLIBC_DEP_GCC=y
 CT_GLIBC_DEP_PYTHON=y
-
-#
-# glibc >=2.29 requires python >3.4
-#
 CT_GLIBC_SPARC_ALLOW_V7=y
 CT_THREADS="nptl"
 CT_GLIBC_BUILD_SSP=y
@@ -522,6 +527,7 @@ CT_GLIBC_SSP_DEFAULT=y
 # CT_GLIBC_SSP_YES is not set
 # CT_GLIBC_SSP_ALL is not set
 # CT_GLIBC_SSP_STRONG is not set
+CT_GLIBC_ENABLE_WERROR=y
 CT_GLIBC_ENABLE_COMMON_FLAG=y
 CT_ALL_LIBC_CHOICES="AVR_LIBC GLIBC MINGW_W64 MOXIEBOX MUSL NEWLIB NONE PICOLIBC UCLIBC_NG"
 CT_LIBC_SUPPORT_THREADS_ANY=y
@@ -564,32 +570,32 @@ CT_GCC_PKG_NAME="gcc"
 CT_GCC_SRC_RELEASE=y
 # CT_GCC_SRC_DEVEL is not set
 CT_GCC_PATCH_ORDER="global"
-CT_GCC_V_13=y
+# CT_GCC_V_13 is not set
 # CT_GCC_V_12 is not set
 # CT_GCC_V_11 is not set
 # CT_GCC_V_10 is not set
 # CT_GCC_V_9 is not set
-# CT_GCC_V_8 is not set
+CT_GCC_V_8=y
 # CT_GCC_V_7 is not set
 # CT_GCC_V_6 is not set
 # CT_GCC_V_5 is not set
 # CT_GCC_V_4_9 is not set
-CT_GCC_VERSION="13.2.0"
+CT_GCC_VERSION="8.5.0"
 CT_GCC_MIRRORS="$(CT_Mirrors GNU gcc/gcc-${CT_GCC_VERSION}) $(CT_Mirrors sourceware gcc/releases/gcc-${CT_GCC_VERSION})"
 CT_GCC_ARCHIVE_FILENAME="@{pkg_name}-@{version}"
 CT_GCC_ARCHIVE_DIRNAME="@{pkg_name}-@{version}"
 CT_GCC_ARCHIVE_FORMATS=".tar.xz .tar.gz"
 CT_GCC_SIGNATURE_FORMAT=""
-CT_GCC_later_than_13=y
-CT_GCC_13_or_later=y
-CT_GCC_later_than_12=y
-CT_GCC_12_or_later=y
-CT_GCC_later_than_11=y
-CT_GCC_11_or_later=y
-CT_GCC_later_than_10=y
-CT_GCC_10_or_later=y
-CT_GCC_later_than_9=y
-CT_GCC_9_or_later=y
+CT_GCC_13_or_older=y
+CT_GCC_older_than_13=y
+CT_GCC_12_or_older=y
+CT_GCC_older_than_12=y
+CT_GCC_11_or_older=y
+CT_GCC_older_than_11=y
+CT_GCC_10_or_older=y
+CT_GCC_older_than_10=y
+CT_GCC_9_or_older=y
+CT_GCC_older_than_9=y
 CT_GCC_later_than_8=y
 CT_GCC_8_or_later=y
 CT_GCC_later_than_7=y
@@ -630,7 +636,6 @@ CT_CC_GCC_LIBSTDCXX_VERBOSE=m
 # Misc. obscure options.
 #
 CT_CC_CXA_ATEXIT=y
-CT_CC_GCC_TM_CLONE_REGISTRY=m
 CT_CC_GCC_DISABLE_PCH=y
 CT_CC_GCC_SJLJ_EXCEPTIONS=m
 CT_CC_GCC_LDBL_128=m
diff --git a/Engine/Build/BatchFiles/Linux/Toolchain/build_linux_toolchain.sh b/Engine/Build/BatchFiles/Linux/Toolchain/build_linux_toolchain.sh
--- a/Engine/Build/BatchFiles/Linux/Toolchain/build_linux_toolchain.sh
+++ b/Engine/Build/BatchFiles/Linux/Toolchain/build_linux_toolchain.sh
@@ -3,9 +3,9 @@
 set -x
 set -eu
 
-ToolChainVersion=v23
-LLVM_VERSION_MAJOR=18
-LLVM_VERSION=${LLVM_VERSION_MAJOR}.1.0
+ToolChainVersion=v24
+LLVM_VERSION_MAJOR=19
+LLVM_VERSION=${LLVM_VERSION_MAJOR}.1.7
 LLVM_BRANCH=release/${LLVM_VERSION_MAJOR}.x
 LLVM_TAG=llvmorg-${LLVM_VERSION}
 
@@ -20,6 +20,8 @@ OutputDirLinux=/src/build/OUTPUT-linux
 OutputDirWindows=/src/build/OUTPUT-windows
 InstallClangDir=/src/build/install-clang
 
+DirsToDelete=
+
 # Default permissions
 umask 0022
 
@@ -109,8 +111,12 @@ if [ ! -f build-clang.done ]; then
 		#   4.8.5.  You can temporarily opt out using
 		#   LLVM_TEMPORARILY_ALLOW_OLD_TOOLCHAIN, but very soon your toolchain won't be
 		#   supported.
-		cmake3 -G "Unix Makefiles" ../llvm-src/llvm \
-			-DLLVM_ENABLE_PROJECTS=llvm\;clang\;lld\;compiler-rt \
+
+		# libunwind in LLVM_ENABLE_RUNTIMES is mutually exclusive with -DLIBCXXABI_USE_LLVM_UNWINDER=OFF
+		cmake3 -G "Unix Makefiles" -S ../llvm-src/llvm -B .  \
+			-DLLVM_ENABLE_PROJECTS="llvm;clang;lld;compiler-rt" \
+			-DLLVM_ENABLE_RUNTIMES="libcxx;libcxxabi" \
+			-DLIBCXXABI_USE_LLVM_UNWINDER=OFF \
 			-DCMAKE_BUILD_TYPE=Release \
 			-DLLVM_ENABLE_TERMINFO=OFF \
 			-DLLVM_ENABLE_LIBXML2=OFF \
@@ -152,7 +158,10 @@ if [ ! -f copy-linux-toolchain.done ]; then
 
 				[[ -f build.log.bz2 ]] && mv build.log.bz2 ../../build-linux-$arch.log.bz2
 
-				rm -rf $arch
+				# don't remove for now so that I can update the copy operations
+				if [ ! -f /src/debug_flow_on ]; then
+					DirsToDelete="${DirsToDelete} ${PWD}/$arch"
+				fi
 			fi
 		popd
 
@@ -169,6 +178,11 @@ if [ ! -f copy-linux-toolchain.done ]; then
 
 		if [ "$arch" == "x86_64-unknown-linux-gnu" ]; then
 			cp -r -L ${InstallClangDir}/lib/clang ${OutputDirLinux}/$arch/lib/
+
+			# copy libc++.a and friends to lib64.  This will need to move out to where the libc++ stuff currently lives in the tree
+			mkdir -p install-libc++-$arch/{lib,include/c++}
+			cp -L ${InstallClangDir}/lib/$arch/* install-libc++-$arch/lib
+			cp -r -L ${InstallClangDir}/include/$arch/c++/v1 install-libc++-$arch/include/c++
 		fi
 	done
 	touch copy-linux-toolchain.done
@@ -224,11 +238,40 @@ for arch in $TARGETS; do
 
 	echo "Copying compiler rt..."
 	cp -r install-rt-$arch/lib/* ${OutputDirLinux}/$arch/lib/clang/${LLVM_VERSION_MAJOR}/lib/
+
+
+	echo "Building libc++..."
+	mkdir -p build-libc++-$arch
+	pushd build-libc++-$arch
+
+		cmake3 -G "Unix Makefiles" ../llvm-src/runtimes \
+			-DLLVM_ENABLE_RUNTIMES="libcxx;libcxxabi" \
+			-DLIBCXXABI_USE_LLVM_UNWINDER=OFF \
+			-DCMAKE_BUILD_TYPE=Release \
+			-DCMAKE_SYSTEM_NAME="Linux" \
+			-DCMAKE_C_COMPILER_TARGET="$arch" \
+			-DCMAKE_C_COMPILER=${InstallClangDir}/bin/clang \
+			-DCMAKE_CXX_COMPILER_TARGET="$arch" \
+			-DCMAKE_CXX_COMPILER=${InstallClangDir}/bin/clang++ \
+			-DCMAKE_AR=${InstallClangDir}/bin/llvm-ar \
+			-DCMAKE_NM=${InstallClangDir}/bin/llvm-nm \
+			-DCMAKE_RANLIB=${InstallClangDir}/bin/llvm-ranlib \
+			-DCMAKE_EXE_LINKER_FLAGS="--target=$arch -L${OutputDirLinux}/$arch/lib64 --sysroot=${OutputDirLinux}/$arch -fuse-ld=lld" \
+			-DCMAKE_C_FLAGS="--target=$arch --sysroot=${OutputDirLinux}/$arch" \
+			-DCMAKE_CXX_FLAGS="--target=$arch --sysroot=${OutputDirLinux}/$arch" \
+			-DCMAKE_ASM_FLAGS="--target=$arch --sysroot=${OutputDirLinux}/$arch" \
+			-DCMAKE_INSTALL_PREFIX=../install-libc++-$arch
+
+		make -j$CORES && make install
+
+	popd
+
 done
 
 touch build-runtime.done
 fi
 
+
 # Create version file
 echo "${ToolChainVersionName}" > ${OutputDirLinux}/ToolchainVersion.txt
 
@@ -263,12 +306,26 @@ if [ ! -f copy-windows-toolchain.done ]; then
 
 			[[ -f build.log.bz2 ]] && mv build.log.bz2 ../../build-windows-$arch.log.bz2
 
-			rm -rf $arch
+			if [ ! -f /src/debug_flow_on ]; then
+				DirsToDelete="${DirsToDelete} ${PWD}/$arch"
+			fi
 		popd
 	done
 	touch copy-windows-toolchain.done
 fi
 
+
+# clean up toolchain remnants before packaging
+if [ ! -f /src/debug_flow_on ]; then
+	echo "Cleaning up..." 
+	for dir in $DirsToDelete; do
+			echo "   Removing $dir..."
+			rm -rf $dir
+	done    
+else
+	echo "Skipping cleanup because 'debug_flow_on' exists..."
+fi
+
 # Pack Linux files
 pushd ${OutputDirLinux}
 	mkdir -p build/{src,scripts}
@@ -307,4 +364,7 @@ pushd ${OutputDirWindows}
 	zip -r /src/build/${ToolChainVersionName}-windows.zip *
 popd
 
+
+echo "Don't forget to run update_libc++.sh and perform a p4 reconcile in Engine/Source/ThirdParty/Unix/LibCxx..."
+
 echo done.
diff --git a/Engine/Build/BatchFiles/Linux/Toolchain/update_libc++.sh b/Engine/Build/BatchFiles/Linux/Toolchain/update_libc++.sh
new file mode 100755
--- /dev/null
+++ b/Engine/Build/BatchFiles/Linux/Toolchain/update_libc++.sh
@@ -0,0 +1,16 @@
+#!/bin/bash
+
+set -eu -x
+
+SCRIPT_DIR=$(cd "$(dirname "$BASH_SOURCE")" ; pwd)
+
+# purge old LibC++ files from tree
+rm -rf ${SCRIPT_DIR}/../../../../Source/ThirdParty/Unix/LibCxx/include/*
+find ${SCRIPT_DIR}/../../../../Source/ThirdParty/Unix/LibCxx/lib -type f -delete
+
+# copy new LibC++ files into place
+# grab aarch64 includes because the x86_64 ones are identical, but get built in a weird place due to how they're built
+cp -rf build/install-libc++-aarch64-unknown-linux-gnueabi/include ${SCRIPT_DIR}/../../../../Source/ThirdParty/Unix/LibCxx
+cp build/install-libc++-x86_64-unknown-linux-gnu/lib/{libc++.a,libc++abi.a} ${SCRIPT_DIR}/../../../../Source/ThirdParty/Unix/LibCxx/lib/Unix/x86_64-unknown-linux-gnu
+cp build/install-libc++-aarch64-unknown-linux-gnueabi/lib/{libc++.a,libc++abi.a} ${SCRIPT_DIR}/../../../../Source/ThirdParty/Unix/LibCxx/lib/Unix/aarch64-unknown-linux-gnueabi
+
diff --git a/Engine/Build/BatchFiles/Linux/Toolchain/x86_64-unknown-linux-gnu.linux.config b/Engine/Build/BatchFiles/Linux/Toolchain/x86_64-unknown-linux-gnu.linux.config
--- a/Engine/Build/BatchFiles/Linux/Toolchain/x86_64-unknown-linux-gnu.linux.config
+++ b/Engine/Build/BatchFiles/Linux/Toolchain/x86_64-unknown-linux-gnu.linux.config
@@ -15,6 +15,7 @@ CT_CONFIGURE_has_autoconf_2_65_or_newer=y
 CT_CONFIGURE_has_autoreconf_2_65_or_newer=y
 CT_CONFIGURE_has_automake_1_15_or_newer=y
 CT_CONFIGURE_has_gnu_m4_1_4_12_or_newer=y
+CT_CONFIGURE_has_python_3_4_or_newer=y
 CT_CONFIGURE_has_bison_2_7_or_newer=y
 CT_CONFIGURE_has_bison_3_0_4_or_newer=y
 CT_CONFIGURE_has_python=y
@@ -412,6 +413,16 @@ CT_GLIBC_PKG_NAME="glibc"
 CT_GLIBC_SRC_RELEASE=y
 # CT_GLIBC_SRC_DEVEL is not set
 CT_GLIBC_PATCH_ORDER="global"
+# CT_GLIBC_V_2_38 is not set
+# CT_GLIBC_V_2_37 is not set
+# CT_GLIBC_V_2_36 is not set
+# CT_GLIBC_V_2_35 is not set
+# CT_GLIBC_V_2_34 is not set
+# CT_GLIBC_V_2_33 is not set
+# CT_GLIBC_V_2_32 is not set
+# CT_GLIBC_V_2_31 is not set
+# CT_GLIBC_V_2_30 is not set
+# CT_GLIBC_V_2_29 is not set
 CT_GLIBC_V_2_28=y
 # CT_GLIBC_V_2_27 is not set
 # CT_GLIBC_V_2_26 is not set
@@ -442,7 +453,6 @@ CT_GLIBC_2_30_or_older=y
 CT_GLIBC_older_than_2_30=y
 CT_GLIBC_2_29_or_older=y
 CT_GLIBC_older_than_2_29=y
-CT_GLIBC_REQUIRE_older_than_2_29=y
 CT_GLIBC_2_28_or_later=y
 CT_GLIBC_2_28_or_older=y
 CT_GLIBC_later_than_2_27=y
@@ -465,10 +475,6 @@ CT_GLIBC_DEP_KERNEL_HEADERS_VERSION=y
 CT_GLIBC_DEP_BINUTILS=y
 CT_GLIBC_DEP_GCC=y
 CT_GLIBC_DEP_PYTHON=y
-
-#
-# glibc >=2.29 requires python >3.4
-#
 CT_GLIBC_SPARC_ALLOW_V7=y
 CT_THREADS="nptl"
 CT_GLIBC_BUILD_SSP=y
@@ -494,6 +500,7 @@ CT_GLIBC_SSP_DEFAULT=y
 # CT_GLIBC_SSP_YES is not set
 # CT_GLIBC_SSP_ALL is not set
 # CT_GLIBC_SSP_STRONG is not set
+CT_GLIBC_ENABLE_WERROR=y
 # CT_GLIBC_ENABLE_COMMON_FLAG is not set
 CT_ALL_LIBC_CHOICES="AVR_LIBC GLIBC MINGW_W64 MOXIEBOX MUSL NEWLIB NONE PICOLIBC UCLIBC_NG"
 CT_LIBC_SUPPORT_THREADS_ANY=y
@@ -536,32 +543,32 @@ CT_GCC_PKG_NAME="gcc"
 CT_GCC_SRC_RELEASE=y
 # CT_GCC_SRC_DEVEL is not set
 CT_GCC_PATCH_ORDER="global"
-CT_GCC_V_13=y
+# CT_GCC_V_13 is not set
 # CT_GCC_V_12 is not set
 # CT_GCC_V_11 is not set
 # CT_GCC_V_10 is not set
 # CT_GCC_V_9 is not set
-# CT_GCC_V_8 is not set
+CT_GCC_V_8=y
 # CT_GCC_V_7 is not set
 # CT_GCC_V_6 is not set
 # CT_GCC_V_5 is not set
 # CT_GCC_V_4_9 is not set
-CT_GCC_VERSION="13.2.0"
+CT_GCC_VERSION="8.5.0"
 CT_GCC_MIRRORS="$(CT_Mirrors GNU gcc/gcc-${CT_GCC_VERSION}) $(CT_Mirrors sourceware gcc/releases/gcc-${CT_GCC_VERSION})"
 CT_GCC_ARCHIVE_FILENAME="@{pkg_name}-@{version}"
 CT_GCC_ARCHIVE_DIRNAME="@{pkg_name}-@{version}"
 CT_GCC_ARCHIVE_FORMATS=".tar.xz .tar.gz"
 CT_GCC_SIGNATURE_FORMAT=""
-CT_GCC_later_than_13=y
-CT_GCC_13_or_later=y
-CT_GCC_later_than_12=y
-CT_GCC_12_or_later=y
-CT_GCC_later_than_11=y
-CT_GCC_11_or_later=y
-CT_GCC_later_than_10=y
-CT_GCC_10_or_later=y
-CT_GCC_later_than_9=y
-CT_GCC_9_or_later=y
+CT_GCC_13_or_older=y
+CT_GCC_older_than_13=y
+CT_GCC_12_or_older=y
+CT_GCC_older_than_12=y
+CT_GCC_11_or_older=y
+CT_GCC_older_than_11=y
+CT_GCC_10_or_older=y
+CT_GCC_older_than_10=y
+CT_GCC_9_or_older=y
+CT_GCC_older_than_9=y
 CT_GCC_later_than_8=y
 CT_GCC_8_or_later=y
 CT_GCC_later_than_7=y
@@ -603,7 +610,6 @@ CT_CC_GCC_LIBSTDCXX_VERBOSE=m
 # Misc. obscure options.
 #
 CT_CC_CXA_ATEXIT=y
-CT_CC_GCC_TM_CLONE_REGISTRY=m
 CT_CC_GCC_DISABLE_PCH=y
 CT_CC_GCC_SJLJ_EXCEPTIONS=m
 CT_CC_GCC_LDBL_128=m
diff --git a/Engine/Build/BatchFiles/Linux/Toolchain/x86_64-unknown-linux-gnu.windows.config b/Engine/Build/BatchFiles/Linux/Toolchain/x86_64-unknown-linux-gnu.windows.config
--- a/Engine/Build/BatchFiles/Linux/Toolchain/x86_64-unknown-linux-gnu.windows.config
+++ b/Engine/Build/BatchFiles/Linux/Toolchain/x86_64-unknown-linux-gnu.windows.config
@@ -15,6 +15,7 @@ CT_CONFIGURE_has_autoconf_2_65_or_newer=y
 CT_CONFIGURE_has_autoreconf_2_65_or_newer=y
 CT_CONFIGURE_has_automake_1_15_or_newer=y
 CT_CONFIGURE_has_gnu_m4_1_4_12_or_newer=y
+CT_CONFIGURE_has_python_3_4_or_newer=y
 CT_CONFIGURE_has_bison_2_7_or_newer=y
 CT_CONFIGURE_has_bison_3_0_4_or_newer=y
 CT_CONFIGURE_has_python=y
@@ -417,6 +418,16 @@ CT_GLIBC_PKG_NAME="glibc"
 CT_GLIBC_SRC_RELEASE=y
 # CT_GLIBC_SRC_DEVEL is not set
 CT_GLIBC_PATCH_ORDER="global"
+# CT_GLIBC_V_2_38 is not set
+# CT_GLIBC_V_2_37 is not set
+# CT_GLIBC_V_2_36 is not set
+# CT_GLIBC_V_2_35 is not set
+# CT_GLIBC_V_2_34 is not set
+# CT_GLIBC_V_2_33 is not set
+# CT_GLIBC_V_2_32 is not set
+# CT_GLIBC_V_2_31 is not set
+# CT_GLIBC_V_2_30 is not set
+# CT_GLIBC_V_2_29 is not set
 CT_GLIBC_V_2_28=y
 # CT_GLIBC_V_2_27 is not set
 # CT_GLIBC_V_2_26 is not set
@@ -447,7 +458,6 @@ CT_GLIBC_2_30_or_older=y
 CT_GLIBC_older_than_2_30=y
 CT_GLIBC_2_29_or_older=y
 CT_GLIBC_older_than_2_29=y
-CT_GLIBC_REQUIRE_older_than_2_29=y
 CT_GLIBC_2_28_or_later=y
 CT_GLIBC_2_28_or_older=y
 CT_GLIBC_later_than_2_27=y
@@ -470,10 +480,6 @@ CT_GLIBC_DEP_KERNEL_HEADERS_VERSION=y
 CT_GLIBC_DEP_BINUTILS=y
 CT_GLIBC_DEP_GCC=y
 CT_GLIBC_DEP_PYTHON=y
-
-#
-# glibc >=2.29 requires python >3.4
-#
 CT_GLIBC_SPARC_ALLOW_V7=y
 CT_THREADS="nptl"
 CT_GLIBC_BUILD_SSP=y
@@ -499,6 +505,7 @@ CT_GLIBC_SSP_DEFAULT=y
 # CT_GLIBC_SSP_YES is not set
 # CT_GLIBC_SSP_ALL is not set
 # CT_GLIBC_SSP_STRONG is not set
+CT_GLIBC_ENABLE_WERROR=y
 CT_GLIBC_ENABLE_COMMON_FLAG=y
 CT_ALL_LIBC_CHOICES="AVR_LIBC GLIBC MINGW_W64 MOXIEBOX MUSL NEWLIB NONE PICOLIBC UCLIBC_NG"
 CT_LIBC_SUPPORT_THREADS_ANY=y
@@ -541,32 +548,32 @@ CT_GCC_PKG_NAME="gcc"
 CT_GCC_SRC_RELEASE=y
 # CT_GCC_SRC_DEVEL is not set
 CT_GCC_PATCH_ORDER="global"
-CT_GCC_V_13=y
+# CT_GCC_V_13 is not set
 # CT_GCC_V_12 is not set
 # CT_GCC_V_11 is not set
 # CT_GCC_V_10 is not set
 # CT_GCC_V_9 is not set
-# CT_GCC_V_8 is not set
+CT_GCC_V_8=y
 # CT_GCC_V_7 is not set
 # CT_GCC_V_6 is not set
 # CT_GCC_V_5 is not set
 # CT_GCC_V_4_9 is not set
-CT_GCC_VERSION="13.2.0"
+CT_GCC_VERSION="8.5.0"
 CT_GCC_MIRRORS="$(CT_Mirrors GNU gcc/gcc-${CT_GCC_VERSION}) $(CT_Mirrors sourceware gcc/releases/gcc-${CT_GCC_VERSION})"
 CT_GCC_ARCHIVE_FILENAME="@{pkg_name}-@{version}"
 CT_GCC_ARCHIVE_DIRNAME="@{pkg_name}-@{version}"
 CT_GCC_ARCHIVE_FORMATS=".tar.xz .tar.gz"
 CT_GCC_SIGNATURE_FORMAT=""
-CT_GCC_later_than_13=y
-CT_GCC_13_or_later=y
-CT_GCC_later_than_12=y
-CT_GCC_12_or_later=y
-CT_GCC_later_than_11=y
-CT_GCC_11_or_later=y
-CT_GCC_later_than_10=y
-CT_GCC_10_or_later=y
-CT_GCC_later_than_9=y
-CT_GCC_9_or_later=y
+CT_GCC_13_or_older=y
+CT_GCC_older_than_13=y
+CT_GCC_12_or_older=y
+CT_GCC_older_than_12=y
+CT_GCC_11_or_older=y
+CT_GCC_older_than_11=y
+CT_GCC_10_or_older=y
+CT_GCC_older_than_10=y
+CT_GCC_9_or_older=y
+CT_GCC_older_than_9=y
 CT_GCC_later_than_8=y
 CT_GCC_8_or_later=y
 CT_GCC_later_than_7=y
@@ -608,7 +615,6 @@ CT_CC_GCC_LIBSTDCXX_VERBOSE=m
 # Misc. obscure options.
 #
 CT_CC_CXA_ATEXIT=y
-CT_CC_GCC_TM_CLONE_REGISTRY=m
 CT_CC_GCC_DISABLE_PCH=y
 CT_CC_GCC_SJLJ_EXCEPTIONS=m
 CT_CC_GCC_LDBL_128=m
