Linux: Find GDB initialization file in ~/.config too

Initialization files: https://sourceware.org/gdb/current/onlinedocs/gdb.html/Initialization-Files.html#Home-Directory-Init-File

[RN] Added additional gdbinit search paths

diff --git a/Engine/Build/BatchFiles/Linux/GenerateGDBInit.sh b/Engine/Build/BatchFiles/Linux/GenerateGDBInit.sh
--- a/Engine/Build/BatchFiles/Linux/GenerateGDBInit.sh
+++ b/Engine/Build/BatchFiles/Linux/GenerateGDBInit.sh
@@ -18,14 +18,22 @@ AddGDBPrettyPrinters()
 	chmod 644 ~/.config/Epic/GDBPrinters/UEPrinters.py 	# set R/W again (it can be read-only if copied from Perforce)
 
 	# Check if .gdbinit exists. If not create else add needed parts.
-	if [ ! -f ~/.gdbinit ]; then
-		echo "no ~/.gdbinit file found - creating a new one."
-		echo -e "python \nimport sys\n\nsys.path.append('$HOME/.config/Epic/GDBPrinters/')\n\nfrom UEPrinters import register_ue_printers\nregister_ue_printers(None)\nprint(\"Registered pretty printers for UE classes\")\n\nend" >> ~/.gdbinit
+	local gdbinit
+	for path in "$XDG_CONFIG_HOME/gdb/gdbinit" "$HOME/.config/gdb/gdbinit" "$HOME/.gdbinit"; do
+		gdbinit="$path"
+		if [ -f "$gdbinit" ]; then
+			break
+		fi
+	done
+
+	if [ ! -f "$gdbinit" ]; then
+		echo "no GDB initialization file found - creating a new one."
+		echo -e "python \nimport sys\n\nsys.path.append('$HOME/.config/Epic/GDBPrinters/')\n\nfrom UEPrinters import register_ue_printers\nregister_ue_printers(None)\nprint(\"Registered pretty printers for UE classes\")\n\nend" >> "$gdbinit"
 	else
-		if grep -q "register_ue_printers" ~/.gdbinit; then
-			echo "found necessary entries in ~/.gdbinit file, not changing it."
+		if grep -q "register_ue_printers" "$gdbinit"; then
+			echo "found necessary entries in $gdbinit, not changing it."
 		else
-			echo -e "cannot modify .gdbinit. Please add the below lines manually:\n\n"
+			echo -e "cannot modify $gdbinit. Please add the below lines manually:\n\n"
 			echo -e "python"
 			echo -e "\timport sys"
 			echo -e "\tsys.path.append('$HOME/.config/Epic/GDBPrinters/')"
