add -fPIC for arm64 libc++.a and libc++abi.a

diff --git a/Engine/Build/BatchFiles/Linux/Toolchain/build_linux_toolchain.sh b/Engine/Build/BatchFiles/Linux/Toolchain/build_linux_toolchain.sh
--- a/Engine/Build/BatchFiles/Linux/Toolchain/build_linux_toolchain.sh
+++ b/Engine/Build/BatchFiles/Linux/Toolchain/build_linux_toolchain.sh
@@ -244,23 +244,25 @@ for arch in $TARGETS; do
 	mkdir -p build-libc++-$arch
 	pushd build-libc++-$arch
 
-		cmake3 -G "Unix Makefiles" ../llvm-src/runtimes \
-			-DLLVM_ENABLE_RUNTIMES="libcxx;libcxxabi" \
-			-DLIBCXXABI_USE_LLVM_UNWINDER=OFF \
-			-DCMAKE_BUILD_TYPE=Release \
-			-DCMAKE_SYSTEM_NAME="Linux" \
-			-DCMAKE_C_COMPILER_TARGET="$arch" \
-			-DCMAKE_C_COMPILER=${InstallClangDir}/bin/clang \
-			-DCMAKE_CXX_COMPILER_TARGET="$arch" \
-			-DCMAKE_CXX_COMPILER=${InstallClangDir}/bin/clang++ \
-			-DCMAKE_AR=${InstallClangDir}/bin/llvm-ar \
-			-DCMAKE_NM=${InstallClangDir}/bin/llvm-nm \
-			-DCMAKE_RANLIB=${InstallClangDir}/bin/llvm-ranlib \
-			-DCMAKE_EXE_LINKER_FLAGS="--target=$arch -L${OutputDirLinux}/$arch/lib64 --sysroot=${OutputDirLinux}/$arch -fuse-ld=lld" \
-			-DCMAKE_C_FLAGS="--target=$arch --sysroot=${OutputDirLinux}/$arch" \
-			-DCMAKE_CXX_FLAGS="--target=$arch --sysroot=${OutputDirLinux}/$arch" \
-			-DCMAKE_ASM_FLAGS="--target=$arch --sysroot=${OutputDirLinux}/$arch" \
-			-DCMAKE_INSTALL_PREFIX=../install-libc++-$arch
+                cmake3 -G "Unix Makefiles" ../llvm-src/runtimes \
+                        -DLLVM_ENABLE_RUNTIMES="libcxx;libcxxabi" \
+                        -DLIBCXXABI_USE_LLVM_UNWINDER=OFF \
+                        -DLIBCXX_ENABLE_SHARED=OFF \
+                        -DLIBCXXABI_ENABLE_SHARED=OFF \
+                        -DCMAKE_BUILD_TYPE=Release \
+                        -DCMAKE_SYSTEM_NAME="Linux" \
+                        -DCMAKE_C_COMPILER_TARGET="$arch" \
+                        -DCMAKE_C_COMPILER=${InstallClangDir}/bin/clang \
+                        -DCMAKE_CXX_COMPILER_TARGET="$arch" \
+                        -DCMAKE_CXX_COMPILER=${InstallClangDir}/bin/clang++ \
+                        -DCMAKE_AR=${InstallClangDir}/bin/llvm-ar \
+                        -DCMAKE_NM=${InstallClangDir}/bin/llvm-nm \
+                        -DCMAKE_RANLIB=${InstallClangDir}/bin/llvm-ranlib \
+                        -DCMAKE_EXE_LINKER_FLAGS="-fPIC --target=$arch -L${OutputDirLinux}/$arch/lib64 --sysroot=${OutputDirLinux}/$arch -fuse-ld=lld" \
+                        -DCMAKE_C_FLAGS="-fPIC --target=$arch --sysroot=${OutputDirLinux}/$arch" \
+                        -DCMAKE_CXX_FLAGS="-fPIC --target=$arch --sysroot=${OutputDirLinux}/$arch" \
+                        -DCMAKE_ASM_FLAGS="-fPIC --target=$arch --sysroot=${OutputDirLinux}/$arch" \
+                        -DCMAKE_INSTALL_PREFIX=../install-libc++-$arch
 
 		make -j$CORES && make install
 
