Handle app extension bundleID change

#rnx

diff --git a/Engine/Build/BatchFiles/Mac/ResignApp.sh b/Engine/Build/BatchFiles/Mac/ResignApp.sh
--- a/Engine/Build/BatchFiles/Mac/ResignApp.sh
+++ b/Engine/Build/BatchFiles/Mac/ResignApp.sh
@@ -81,19 +81,27 @@ DoWork()
 
 	if [[ ! -z "$DEVELOPER" ]]; then
 		EXTRAOPTIONS=""
+		EXTAPPBUNDLEID=""
 		while IFS='' read -r line || [[ -n "$line" ]]; do
 	
 			if [[ ! -z "$BUNDLE" ]] && [[ "$line" == *".appex"* ]]; then
 			   extbundleid=$(/usr/libexec/PlistBuddy -c "Print:CFBundleIdentifier" "$line/Info.plist")
-			   echo "Changing .appex BundleID $extbundleid with : ${extbundleid/$oldbundleid/$BUNDLE}"
-			   /usr/libexec/PlistBuddy -c "Set:CFBundleIdentifier ${extbundleid/$oldbundleid/$BUNDLE}" "$line/Info.plist"
+			   EXTAPPBUNDLEID=${extbundleid/$oldbundleid/$BUNDLE}
+			   echo "Changing .appex BundleID $extbundleid with : $EXTAPPBUNDLEID"
+			   /usr/libexec/PlistBuddy -c "Set:CFBundleIdentifier $EXTAPPBUNDLEID" "$line/Info.plist"
 			fi    
 	
 			if [[ -f "$line/embedded.mobileprovision" ]]; then
 				security cms -D -i "$line/embedded.mobileprovision" > t_entitlements_full.plist
 				/usr/libexec/PlistBuddy -x -c 'Print:Entitlements' t_entitlements_full.plist > t_entitlements.plist
 				TEAMID=$(/usr/libexec/PlistBuddy -c 'Print:com.apple.developer.team-identifier' t_entitlements.plist)
-				/usr/libexec/PlistBuddy -c "Set:application-identifier $TEAMID.$BUNDLE" t_entitlements.plist
+
+				if [[ ! -z "$BUNDLE" ]] && [[ "$line" == *".appex"* ]]; then
+					echo "Updating .appex embedded.mobileprovision bundleID with : $TEAMID.$EXTAPPBUNDLEID"
+					/usr/libexec/PlistBuddy -c "Set:application-identifier $TEAMID.$EXTAPPBUNDLEID" t_entitlements.plist
+				else
+					/usr/libexec/PlistBuddy -c "Set:application-identifier $TEAMID.$BUNDLE" t_entitlements.plist
+				fi    
 				
 				if [[ ! -z "$BUNDLE" ]]; then
 					EXTRAOPTIONS="--entitlements t_entitlements.plist"
