Update iOS ResignApp script to not resign frameworks.  Was reusing the previous loops params.

#rnx

diff --git a/Engine/Build/BatchFiles/Mac/ResignApp.sh b/Engine/Build/BatchFiles/Mac/ResignApp.sh
--- a/Engine/Build/BatchFiles/Mac/ResignApp.sh
+++ b/Engine/Build/BatchFiles/Mac/ResignApp.sh
@@ -100,6 +100,7 @@ DoWork()
 					echo "Updating .appex embedded.mobileprovision bundleID with : $TEAMID.$EXTAPPBUNDLEID"
 					/usr/libexec/PlistBuddy -c "Set:application-identifier $TEAMID.$EXTAPPBUNDLEID" t_entitlements.plist
 				else
+					echo "Updating app embedded.mobileprovision bundleID with : $TEAMID.$BUNDLE"
 					/usr/libexec/PlistBuddy -c "Set:application-identifier $TEAMID.$BUNDLE" t_entitlements.plist
 				fi    
 				
@@ -113,6 +114,12 @@ DoWork()
 			echo ""
 			echo Codesigning $line [/usr/bin/codesign  --continue -f -s "$DEVELOPER" $EXTRAOPTIONS "$line"]
 			/usr/bin/codesign --continue -f -s "$DEVELOPER" $EXTRAOPTIONS "$line"
+
+			# reset for next loop
+			EXTRAOPTIONS=""
+			EXTAPPBUNDLEID=""
+			rm -f t_entitlements_full.plist
+			rm -f t_entitlements.plist
 	
 		done < directories.txt
 	fi
