Allow setting an additional profile for an app extension

#rnx

diff --git a/Engine/Build/BatchFiles/Mac/ResignApp.sh b/Engine/Build/BatchFiles/Mac/ResignApp.sh
--- a/Engine/Build/BatchFiles/Mac/ResignApp.sh
+++ b/Engine/Build/BatchFiles/Mac/ResignApp.sh
@@ -41,7 +41,7 @@ DoWork()
 	fi
 
 	if [[ ! -z "$MOBILEPROV" ]]; then
-		echo "Prep: Copying $SOURCEAPP into app..."
+		echo "Prep: Copying $MOBILEPROV profile into app..."
 		cp "$MOBILEPROV" "$TARGET/embedded.mobileprovision"
 	fi
 
@@ -92,6 +92,11 @@ DoWork()
 			fi    
 	
 			if [[ -f "$line/embedded.mobileprovision" ]]; then
+				if [[ ! -z "$BUNDLE" ]] && [[ "$line" == *".appex"* ]] && [[ ! -z "$MOBILEAPPEXPROV" ]]; then
+					echo "Copying appex $MOBILEAPPEXPROV profile into $line.."
+					cp "$MOBILEAPPEXPROV" "$line/embedded.mobileprovision"
+				fi
+
 				security cms -D -i "$line/embedded.mobileprovision" > t_entitlements_full.plist
 				/usr/libexec/PlistBuddy -x -c 'Print:Entitlements' t_entitlements_full.plist > t_entitlements.plist
 				TEAMID=$(/usr/libexec/PlistBuddy -c 'Print:com.apple.developer.team-identifier' t_entitlements.plist)
@@ -187,6 +192,8 @@ Help()
 	echo "OPTIONAL OPTIONS:"
 	echo "  -p | --provision <path to .mobileprovision>"
 	echo "     Path to the .mobileprovision file to sign with"
+	echo "  -pe | --provisionappex <path to .mobileprovision>"
+	echo "     Path to the .mobileprovision file for the app extension to sign with"
 	echo "  -b | --bundleid"
 	echo "     Bundle ID to use in the app, repleacing existing bundle ID"
 	echo "  -mv | --marketingversion"
@@ -225,6 +232,11 @@ while [[ $# -gt 0 ]]; do
 	  shift # past argument
 	  shift # past value
 	  ;;
+	-pe|--provisionappex)
+	  MOBILEAPPEXPROV="$2"
+	  shift # past argument
+	  shift # past value
+	  ;;
 	-d|--destapp)
 	  TARGETAPP="$2"
 	  shift # past argument
