GetDotNet: Improve arm64 check

#rnx

diff --git a/Engine/Build/BatchFiles/GetDotnetPath.bat b/Engine/Build/BatchFiles/GetDotnetPath.bat
--- a/Engine/Build/BatchFiles/GetDotnetPath.bat
+++ b/Engine/Build/BatchFiles/GetDotnetPath.bat
@@ -16,7 +16,7 @@ set UE_DOTNET_ARCH=win-x64
 if "%PROCESSOR_ARCHITECTURE%" == "ARM64" (
 	set UE_DOTNET_ARCH=win-arm64
 ) else if "%PROCESSOR_ARCHITECTURE%" == "AMD64" (
-	if "%ProgramFiles(Arm)%" NEQ "" (
+	if defined ProgramFiles^(Arm^) (
 		set UE_DOTNET_ARCH=win-arm64
 	)
 )
