Fix references added to AutomationTool by GenerateProjectFiles causing build errors

#rnx

diff --git a/Engine/Build/BatchFiles/BuildUAT.bat b/Engine/Build/BatchFiles/BuildUAT.bat
--- a/Engine/Build/BatchFiles/BuildUAT.bat
+++ b/Engine/Build/BatchFiles/BuildUAT.bat
@@ -52,12 +52,12 @@ fc /C "%TEMP_DEPENDS_FILE%" %DEPENDS_FILE% >nul 2>&1 && (
 
 :Build_AutomationTool
 echo Building AutomationTool...
-dotnet build Programs\AutomationTool\AutomationTool.csproj -c Development -v %MSBUILD_LOGLEVEL%
+dotnet build Programs\AutomationTool\AutomationTool.csproj -c Development -p AutomationToolProjectOnly=true -v %MSBUILD_LOGLEVEL%
 if errorlevel 1 goto Error_UATCompileFailed
 
 echo Publishing AutomationTool...
 if not exist ..\Binaries\DotNET\AutomationTool mkdir ..\Binaries\DotNET\AutomationTool >nul 2>&1
-dotnet publish Programs\AutomationTool\AutomationTool.csproj -c Development --output ..\Binaries\DotNET\AutomationTool --no-build -v %MSBUILD_LOGLEVEL%
+dotnet publish Programs\AutomationTool\AutomationTool.csproj -c Development -p AutomationToolProjectOnly=true --output ..\Binaries\DotNET\AutomationTool --no-build -v %MSBUILD_LOGLEVEL%
 if errorlevel 1 goto Error_UATCompileFailed
 
 rem record input files - regardless of how we got here, these are now our point of reference
