BatchFiles: Normalize dotnet path with realpath

#rnx

diff --git a/Engine/Build/BatchFiles/Linux/SetupDotnet.sh b/Engine/Build/BatchFiles/Linux/SetupDotnet.sh
--- a/Engine/Build/BatchFiles/Linux/SetupDotnet.sh
+++ b/Engine/Build/BatchFiles/Linux/SetupDotnet.sh
@@ -44,7 +44,8 @@ if [ $IS_DOTNET_INSTALLED -eq 0 ]; then
 	ARCH=x64
 	[ $(uname -m) == "arm64" ] && ARCH=arm64 
 
-	export UE_DOTNET_DIR="$CUR_DIR/../../../Binaries/ThirdParty/DotNet/8.0.412/linux-$ARCH"
+	UE_DOTNET_DIR=$(realpath "$CUR_DIR/../../../Binaries/ThirdParty/DotNet/8.0.412/linux-$ARCH")
+	export UE_DOTNET_DIR="$UE_DOTNET_DIR"
 	chmod u+x "$UE_DOTNET_DIR/dotnet"
 	export PATH="$UE_DOTNET_DIR:$PATH"
 	export DOTNET_ROOT="$UE_DOTNET_DIR"
diff --git a/Engine/Build/BatchFiles/Mac/SetupDotnet.sh b/Engine/Build/BatchFiles/Mac/SetupDotnet.sh
--- a/Engine/Build/BatchFiles/Mac/SetupDotnet.sh
+++ b/Engine/Build/BatchFiles/Mac/SetupDotnet.sh
@@ -44,7 +44,8 @@ if [ $IS_DOTNET_INSTALLED -eq 0 ]; then
 	ARCH=x64
 	[ $(uname -m) == "arm64" ] && ARCH=arm64 
 	
-	export UE_DOTNET_DIR=$CUR_DIR/../../../Binaries/ThirdParty/DotNet/8.0.412/mac-$ARCH
+	UE_DOTNET_DIR=$(realpath "$CUR_DIR/../../../Binaries/ThirdParty/DotNet/8.0.412/mac-$ARCH")
+	export UE_DOTNET_DIR="$UE_DOTNET_DIR"
 	chmod u+x "$UE_DOTNET_DIR/dotnet"
 	echo $UE_DOTNET_DIR
 	export PATH=$UE_DOTNET_DIR:$PATH
