BatchFiles: Simple retry script for dotnet build commands

#rnx

diff --git a/Engine/Build/BatchFiles/BuildUAT.sh b/Engine/Build/BatchFiles/BuildUAT.sh
--- a/Engine/Build/BatchFiles/BuildUAT.sh
+++ b/Engine/Build/BatchFiles/BuildUAT.sh
@@ -59,7 +59,7 @@ fi
 
 if [ $PERFORM_REBUILD -eq 1 ]; then
   echo "Building AutomationTool..."
-  dotnet build Programs/AutomationTool/AutomationTool.csproj -c Development -v "$MSBuild_Verbosity"
+  "$SCRIPT_DIR/DotnetRetry.sh" build Programs/AutomationTool/AutomationTool.csproj -c Development -v "$MSBuild_Verbosity"
   if [ $? -ne 0 ]; then
     echo "AutomationTool compilation failed"
     popd >/dev/null || exit 1
@@ -67,7 +67,7 @@ if [ $PERFORM_REBUILD -eq 1 ]; then
   fi
 
   echo "Publishing AutomationTool..."
-  dotnet publish Programs/AutomationTool/AutomationTool.csproj -c Development --output ../Binaries/DotNET/AutomationTool --no-build -v "$MSBuild_Verbosity"
+  "$SCRIPT_DIR/DotnetRetry.sh" publish Programs/AutomationTool/AutomationTool.csproj -c Development --output ../Binaries/DotNET/AutomationTool --no-build -v "$MSBuild_Verbosity"
   if [ $? -ne 0 ]; then
     echo "AutomationTool publish failed"
     popd >/dev/null || exit 1
diff --git a/Engine/Build/BatchFiles/BuildUBT.sh b/Engine/Build/BatchFiles/BuildUBT.sh
--- a/Engine/Build/BatchFiles/BuildUBT.sh
+++ b/Engine/Build/BatchFiles/BuildUBT.sh
@@ -60,7 +60,7 @@ fi
 
 if [ $PERFORM_REBUILD -eq 1 ]; then
   echo "Building UnrealBuildTool..."
-  dotnet build Programs/UnrealBuildTool/UnrealBuildTool.csproj -c Development -v "$MSBuild_Verbosity"
+  "$SCRIPT_DIR/DotnetRetry.sh" build Programs/UnrealBuildTool/UnrealBuildTool.csproj -c Development -v "$MSBuild_Verbosity"
   if [ $? -ne 0 ]; then
     echo "UnrealBuildTool compilation failed"
     popd >/dev/null || exit 1
@@ -68,7 +68,7 @@ if [ $PERFORM_REBUILD -eq 1 ]; then
   fi
 
   echo "Publishing UnrealBuildTool..."
-  dotnet publish Programs/UnrealBuildTool/UnrealBuildTool.csproj -c Development --output ../Binaries/DotNET/UnrealBuildTool --no-build -v "$MSBuild_Verbosity"
+  "$SCRIPT_DIR/DotnetRetry.sh" publish Programs/UnrealBuildTool/UnrealBuildTool.csproj -c Development --output ../Binaries/DotNET/UnrealBuildTool --no-build -v "$MSBuild_Verbosity"
   if [ $? -ne 0 ]; then
     echo "UnrealBuildTool publish failed"
     popd >/dev/null || exit 1
diff --git a/Engine/Build/BatchFiles/DotnetDepends.sh b/Engine/Build/BatchFiles/DotnetDepends.sh
--- a/Engine/Build/BatchFiles/DotnetDepends.sh
+++ b/Engine/Build/BatchFiles/DotnetDepends.sh
@@ -32,7 +32,7 @@ DEPENDS_FILE=$2
 MSBuild_Verbosity="${3:-quiet}"
 TEMPDIR=$(mktemp -q -d)
 
-dotnet msbuild "$SLN" -t:Scan -p:Configuration=Development -p:Platform="Any CPU" -p:OutputPath="$TEMPDIR"/ -p:DependsEncoding=Ascii -noLogo -v:"$MSBuild_Verbosity"
+"$SCRIPT_DIR/DotnetRetry.sh" msbuild "$SLN" -t:Scan -p:Configuration=Development -p:Platform="Any CPU" -p:OutputPath="$TEMPDIR"/ -p:DependsEncoding=Ascii -noLogo -v:"$MSBuild_Verbosity"
 cat "$TEMPDIR"/*.dep.csv | sort -u -o "$DEPENDS_FILE"
 rm -rf "$TEMPDIR"
 
diff --git a/Engine/Build/BatchFiles/DotnetRetry.sh b/Engine/Build/BatchFiles/DotnetRetry.sh
new file mode 100755
--- /dev/null
+++ b/Engine/Build/BatchFiles/DotnetRetry.sh
@@ -0,0 +1,56 @@
+#!/bin/bash
+
+## Unreal Engine AutomationTool build script
+## Copyright Epic Games, Inc. All Rights Reserved
+##
+## This script is expecting to exist in the Engine/Build/BatchFiles directory.  It will not work correctly
+## if you copy it to a different location and run it.
+
+## First, make sure the batch file exists in the folder we expect it to.  This is necessary in order to
+## verify that our relative path to the /Engine/Source directory is correct
+SCRIPT_DIR=$(cd "$(dirname "$0")" && pwd)
+pushd "$SCRIPT_DIR/../../Source" >/dev/null || exit 1
+
+if [ ! -f ../Build/BatchFiles/DotnetRetry.sh ]; then
+  echo
+  echo "DotnetRetry ERROR: The script does not appear to be located in the /Engine/Build/BatchFiles directory.  This script must be run from within that directory."
+  echo
+  popd >/dev/null || exit 1
+  exit 1
+fi
+
+if [ "$(uname)" = "Darwin" ]; then
+  # Setup Environment
+  source "$SCRIPT_DIR/Mac/SetupEnvironment.sh" -dotnet "$SCRIPT_DIR/Mac"
+elif [ "$(uname)" = "Linux" ]; then
+  # Setup Environment
+  source "$SCRIPT_DIR/Linux/SetupEnvironment.sh" -dotnet "$SCRIPT_DIR/Linux"
+fi
+
+MAX_RETRIES=3
+RETRY_DELAY=5
+COUNT=0
+ERRORCODE=0
+
+while [ $COUNT -lt $MAX_RETRIES ]; do
+  # Run the dotnet command to retry if a failure occurs
+  dotnet "$@"
+
+  ERRORCODE=$?
+  if [ $ERRORCODE -eq 0 ]; then
+    break
+  fi
+
+  # Command failed, retry up to our max
+  COUNT=$((COUNT + 1))
+  if [ $COUNT -lt $MAX_RETRIES ]; then
+    echo "dotnet command failed with errorcode $ERRORCODE. Retrying in $RETRY_DELAY seconds... ($COUNT/$MAX_RETRIES)"
+    sleep $RETRY_DELAY
+  else
+    echo "dotnet command failed with errorcode $ERRORCODE after $MAX_RETRIES attempts"
+    break
+  fi
+done
+
+popd >/dev/null || exit 1
+exit $ERRORCODE
