Skip chmod +x dotnet if file is already executable. This addresses permission error if UE is run on a read-only file system on Linux.

diff --git a/Engine/Build/BatchFiles/Linux/SetupDotnet.sh b/Engine/Build/BatchFiles/Linux/SetupDotnet.sh
--- a/Engine/Build/BatchFiles/Linux/SetupDotnet.sh
+++ b/Engine/Build/BatchFiles/Linux/SetupDotnet.sh
@@ -46,7 +46,9 @@ if [ $IS_DOTNET_INSTALLED -eq 0 ]; then
 
 	UE_DOTNET_DIR=$(realpath "$CUR_DIR/../../../Binaries/ThirdParty/DotNet/8.0.412/linux-$ARCH")
 	export UE_DOTNET_DIR="$UE_DOTNET_DIR"
-	chmod u+x "$UE_DOTNET_DIR/dotnet"
+	if [ ! -x "$UE_DOTNET_DIR/dotnet" ]; then
+		chmod u+x "$UE_DOTNET_DIR/dotnet"
+	fi
 	export PATH="$UE_DOTNET_DIR:$PATH"
 	export DOTNET_ROOT="$UE_DOTNET_DIR"
 
