[Linux] Update toolchain buildscripts for llvm-20.1.8

also disable unintentional packaging of shared libs for libc++ when making Linux toolchain

diff --git a/Engine/Build/BatchFiles/Linux/Toolchain/BuildLibCxx.sh b/Engine/Build/BatchFiles/Linux/Toolchain/BuildLibCxx.sh
deleted file mode 100755
diff --git a/Engine/Build/BatchFiles/Linux/Toolchain/RunMe.bat b/Engine/Build/BatchFiles/Linux/Toolchain/RunMe.bat
--- a/Engine/Build/BatchFiles/Linux/Toolchain/RunMe.bat
+++ b/Engine/Build/BatchFiles/Linux/Toolchain/RunMe.bat
@@ -1,9 +1,9 @@
 @REM @echo off
 set BUILD_ONLY=0
-set TOOLCHAIN_VERSION=v25
-set LLVM_VERSION=18.1.0
-set LLVM_BRANCH=release/18.x
-set LLVM_TAG=llvmorg-18.1.0
+set TOOLCHAIN_VERSION=v26
+set LLVM_VERSION=20.1.8
+set LLVM_BRANCH=release/20.x
+set LLVM_TAG=llvmorg-20.1.8
 
 set CMAKE_BINARY=%CD%\..\..\..\..\Extras\ThirdPartyNotUE\CMake\bin\cmake.exe
 set PYTHON_BINARY=%CD%\..\..\..\..\Binaries\ThirdParty\Python\Win64\python.exe
@@ -112,4 +112,4 @@ popd
 
 :endscript
 echo.
-echo Done.
\ No newline at end of file
+echo Done.
diff --git a/Engine/Build/BatchFiles/Linux/Toolchain/build_linux_toolchain.sh b/Engine/Build/BatchFiles/Linux/Toolchain/build_linux_toolchain.sh
--- a/Engine/Build/BatchFiles/Linux/Toolchain/build_linux_toolchain.sh
+++ b/Engine/Build/BatchFiles/Linux/Toolchain/build_linux_toolchain.sh
@@ -3,9 +3,9 @@
 set -x
 set -eu
 
-ToolChainVersion=v25
-LLVM_VERSION_MAJOR=18
-LLVM_VERSION=${LLVM_VERSION_MAJOR}.1.0
+ToolChainVersion=v26
+LLVM_VERSION_MAJOR=20
+LLVM_VERSION=${LLVM_VERSION_MAJOR}.1.8
 LLVM_BRANCH=release/${LLVM_VERSION_MAJOR}.x
 LLVM_TAG=llvmorg-${LLVM_VERSION}
 
@@ -117,6 +117,8 @@ if [ ! -f build-clang.done ]; then
 			-DLLVM_ENABLE_PROJECTS="llvm;clang;lld;compiler-rt" \
 			-DLLVM_ENABLE_RUNTIMES="libcxx;libcxxabi" \
 			-DLIBCXXABI_USE_LLVM_UNWINDER=OFF \
+			-DLIBCXX_ENABLE_SHARED=OFF \
+			-DLIBCXXABI_ENABLE_SHARED=OFF \
 			-DCMAKE_BUILD_TYPE=Release \
 			-DLLVM_ENABLE_TERMINFO=OFF \
 			-DLLVM_ENABLE_LIBXML2=OFF \
diff --git a/Engine/Build/BatchFiles/Linux/Toolchain/update_libc++.sh b/Engine/Build/BatchFiles/Linux/Toolchain/update_libc++.sh
deleted file mode 100755
