DotnetRetry.sh Capture and print output only on success for last failure

#rnx

diff --git a/Engine/Build/BatchFiles/DotnetRetry.sh b/Engine/Build/BatchFiles/DotnetRetry.sh
--- a/Engine/Build/BatchFiles/DotnetRetry.sh
+++ b/Engine/Build/BatchFiles/DotnetRetry.sh
@@ -34,10 +34,11 @@ ERRORCODE=0
 
 while [ $COUNT -lt $MAX_RETRIES ]; do
   # Run the dotnet command to retry if a failure occurs
-  dotnet "$@"
+  output=$(dotnet "$@")
 
   ERRORCODE=$?
   if [ $ERRORCODE -eq 0 ]; then
+    echo "$output"
     break
   fi
 
@@ -48,6 +49,7 @@ while [ $COUNT -lt $MAX_RETRIES ]; do
     sleep $RETRY_DELAY
   else
     echo "dotnet command failed with errorcode $ERRORCODE after $MAX_RETRIES attempts"
+    echo "$output"
     break
   fi
 done
