[Partial Backout] - CL46262455
#rnx
Original CL Desc
-----------------------------------------------------------------
Fix Microsoft.CodeAnalysis.CSharp reference which caused a transitive reference to System.Collections.Immutable 9.0.0 rather than 8.0.0

#rnx

diff --git a/Engine/Build/BatchFiles/BuildUAT.bat b/Engine/Build/BatchFiles/BuildUAT.bat
--- a/Engine/Build/BatchFiles/BuildUAT.bat
+++ b/Engine/Build/BatchFiles/BuildUAT.bat
@@ -56,7 +56,6 @@ dotnet build Programs\AutomationTool\AutomationTool.csproj -c Development -p Aut
 if errorlevel 1 goto Error_UATCompileFailed
 
 echo Publishing AutomationTool...
-if exist ..\Binaries\DotNET\AutomationTool rmdir /s /q ..\Binaries\DotNET\AutomationTool >nul 2>&1
 if not exist ..\Binaries\DotNET\AutomationTool mkdir ..\Binaries\DotNET\AutomationTool >nul 2>&1
 dotnet publish Programs\AutomationTool\AutomationTool.csproj -c Development -p AutomationToolProjectOnly=true --output ..\Binaries\DotNET\AutomationTool --no-build -v %MSBUILD_LOGLEVEL%
 if errorlevel 1 goto Error_UATCompileFailed
diff --git a/Engine/Build/BatchFiles/BuildUAT.sh b/Engine/Build/BatchFiles/BuildUAT.sh
--- a/Engine/Build/BatchFiles/BuildUAT.sh
+++ b/Engine/Build/BatchFiles/BuildUAT.sh
@@ -67,12 +67,6 @@ if [ $PERFORM_REBUILD -eq 1 ]; then
   fi
 
   echo "Publishing AutomationTool..."
-  if [ -d ../Binaries/DotNET/AutomationTool ]; then
-    rm -rf ../Binaries/DotNET/AutomationTool >/dev/null
-  fi
-  if [ ! -d ../Binaries/DotNET/AutomationTool ]; then
-    mkdir -p ../Binaries/DotNET/AutomationTool >/dev/null
-  fi
   "$SCRIPT_DIR/DotnetRetry.sh" publish Programs/AutomationTool/AutomationTool.csproj -c Development --output ../Binaries/DotNET/AutomationTool --no-build -v "$MSBuild_Verbosity"
   if [ $? -ne 0 ]; then
     echo "AutomationTool publish failed"
diff --git a/Engine/Build/BatchFiles/BuildUBT.bat b/Engine/Build/BatchFiles/BuildUBT.bat
--- a/Engine/Build/BatchFiles/BuildUBT.bat
+++ b/Engine/Build/BatchFiles/BuildUBT.bat
@@ -52,7 +52,6 @@ dotnet build Programs\UnrealBuildTool\UnrealBuildTool.csproj -c Development -v %
 if errorlevel 1 goto Error_UBTCompileFailed
 
 echo Publishing UnrealBuildTool...
-if exist ..\Binaries\DotNET\UnrealBuildTool rmdir /s /q ..\Binaries\DotNET\UnrealBuildTool >nul 2>&1
 if not exist ..\Binaries\DotNET\UnrealBuildTool mkdir ..\Binaries\DotNET\UnrealBuildTool >nul 2>&1
 dotnet publish Programs\UnrealBuildTool\UnrealBuildTool.csproj -c Development --output ..\Binaries\DotNET\UnrealBuildTool --no-build -v %MSBUILD_LOGLEVEL%
 if errorlevel 1 goto Error_UBTCompileFailed
diff --git a/Engine/Build/BatchFiles/BuildUBT.sh b/Engine/Build/BatchFiles/BuildUBT.sh
--- a/Engine/Build/BatchFiles/BuildUBT.sh
+++ b/Engine/Build/BatchFiles/BuildUBT.sh
@@ -68,12 +68,6 @@ if [ $PERFORM_REBUILD -eq 1 ]; then
   fi
 
   echo "Publishing UnrealBuildTool..."
-  if [ -d ../Binaries/DotNET/UnrealBuildTool ]; then
-    rm -rf ../Binaries/DotNET/UnrealBuildTool >/dev/null
-  fi
-  if [ ! -d ../Binaries/DotNET/UnrealBuildTool ]; then
-    mkdir -p ../Binaries/DotNET/UnrealBuildTool >/dev/null
-  fi
   "$SCRIPT_DIR/DotnetRetry.sh" publish Programs/UnrealBuildTool/UnrealBuildTool.csproj -c Development --output ../Binaries/DotNET/UnrealBuildTool --no-build -v "$MSBuild_Verbosity"
   if [ $? -ne 0 ]; then
     echo "UnrealBuildTool publish failed"
