When cleaning AutomationTool before publishing do not recusively delete sub-directories as this will delete precompiled automation scripts

#rnx

diff --git a/Engine/Build/BatchFiles/BuildUAT.bat b/Engine/Build/BatchFiles/BuildUAT.bat
--- a/Engine/Build/BatchFiles/BuildUAT.bat
+++ b/Engine/Build/BatchFiles/BuildUAT.bat
@@ -56,6 +56,10 @@ dotnet build Programs\AutomationTool\AutomationTool.csproj -c Development -p Aut
 if errorlevel 1 goto Error_UATCompileFailed
 
 echo Publishing AutomationTool...
+if exist ..\Binaries\DotNET\AutomationTool (
+	del /f /q ..\Binaries\DotNET\AutomationTool\* >nul 2>&1
+	rmdir /s /q ..\Binaries\DotNET\AutomationTool\runtimes >nul 2>&1
+)
 if not exist ..\Binaries\DotNET\AutomationTool mkdir ..\Binaries\DotNET\AutomationTool >nul 2>&1
 dotnet publish Programs\AutomationTool\AutomationTool.csproj -c Development -p AutomationToolProjectOnly=true --output ..\Binaries\DotNET\AutomationTool --no-build -v %MSBUILD_LOGLEVEL%
 if errorlevel 1 goto Error_UATCompileFailed
diff --git a/Engine/Build/BatchFiles/BuildUAT.sh b/Engine/Build/BatchFiles/BuildUAT.sh
--- a/Engine/Build/BatchFiles/BuildUAT.sh
+++ b/Engine/Build/BatchFiles/BuildUAT.sh
@@ -67,6 +67,13 @@ if [ $PERFORM_REBUILD -eq 1 ]; then
   fi
 
   echo "Publishing AutomationTool..."
+  if [ -d ../Binaries/DotNET/AutomationTool ]; then
+    rm -f ../Binaries/DotNET/AutomationTool/* >/dev/null
+    rm -rf ../Binaries/DotNET/AutomationTool/runtimes >/dev/null
+  fi
+  if [ ! -d ../Binaries/DotNET/AutomationTool ]; then
+    mkdir -p ../Binaries/DotNET/AutomationTool >/dev/null
+  fi
   "$SCRIPT_DIR/DotnetRetry.sh" publish Programs/AutomationTool/AutomationTool.csproj -c Development --output ../Binaries/DotNET/AutomationTool --no-build -v "$MSBuild_Verbosity"
   if [ $? -ne 0 ]; then
     echo "AutomationTool publish failed"
diff --git a/Engine/Build/BatchFiles/BuildUBT.bat b/Engine/Build/BatchFiles/BuildUBT.bat
--- a/Engine/Build/BatchFiles/BuildUBT.bat
+++ b/Engine/Build/BatchFiles/BuildUBT.bat
@@ -52,6 +52,7 @@ dotnet build Programs\UnrealBuildTool\UnrealBuildTool.csproj -c Development -v %
 if errorlevel 1 goto Error_UBTCompileFailed
 
 echo Publishing UnrealBuildTool...
+if exist ..\Binaries\DotNET\UnrealBuildTool rmdir /s /q ..\Binaries\DotNET\UnrealBuildTool >nul 2>&1
 if not exist ..\Binaries\DotNET\UnrealBuildTool mkdir ..\Binaries\DotNET\UnrealBuildTool >nul 2>&1
 dotnet publish Programs\UnrealBuildTool\UnrealBuildTool.csproj -c Development --output ..\Binaries\DotNET\UnrealBuildTool --no-build -v %MSBUILD_LOGLEVEL%
 if errorlevel 1 goto Error_UBTCompileFailed
diff --git a/Engine/Build/BatchFiles/BuildUBT.sh b/Engine/Build/BatchFiles/BuildUBT.sh
--- a/Engine/Build/BatchFiles/BuildUBT.sh
+++ b/Engine/Build/BatchFiles/BuildUBT.sh
@@ -68,6 +68,12 @@ if [ $PERFORM_REBUILD -eq 1 ]; then
   fi
 
   echo "Publishing UnrealBuildTool..."
+  if [ -d ../Binaries/DotNET/UnrealBuildTool ]; then
+    rm -rf ../Binaries/DotNET/UnrealBuildTool >/dev/null
+  fi
+  if [ ! -d ../Binaries/DotNET/UnrealBuildTool ]; then
+    mkdir -p ../Binaries/DotNET/UnrealBuildTool >/dev/null
+  fi
   "$SCRIPT_DIR/DotnetRetry.sh" publish Programs/UnrealBuildTool/UnrealBuildTool.csproj -c Development --output ../Binaries/DotNET/UnrealBuildTool --no-build -v "$MSBuild_Verbosity"
   if [ $? -ne 0 ]; then
     echo "UnrealBuildTool publish failed"
diff --git a/Engine/Source/Programs/UnrealBuildTool/UnrealBuildTool.cs b/Engine/Source/Programs/UnrealBuildTool/UnrealBuildTool.cs
--- a/Engine/Source/Programs/UnrealBuildTool/UnrealBuildTool.cs
+++ b/Engine/Source/Programs/UnrealBuildTool/UnrealBuildTool.cs
@@ -519,8 +519,7 @@ namespace UnrealBuildTool
 						using (GlobalSingleInstanceMutex XmlConfigMutex = new GlobalSingleInstanceMutex(XmlConfigMutexName, true))
 						{
 							FileReference? XmlConfigCache = Arguments.GetFileReferenceOrDefault("-XmlConfigCache=", null);
-							FileReference? ProjectFile;
-							Utils.TryParseProjectFileArgument(Arguments, Logger, out ProjectFile, false);
+							Utils.TryParseProjectFileArgument(Arguments, Logger, out FileReference? ProjectFile, false);
 							XmlConfig.ReadConfigFiles(XmlConfigCache, ProjectFile?.Directory, Logger);
 						}
 					}
