Don't delete runtimes directory it's not necessary

#rnx

diff --git a/Engine/Build/BatchFiles/BuildUAT.bat b/Engine/Build/BatchFiles/BuildUAT.bat
--- a/Engine/Build/BatchFiles/BuildUAT.bat
+++ b/Engine/Build/BatchFiles/BuildUAT.bat
@@ -56,10 +56,7 @@ dotnet build Programs\AutomationTool\AutomationTool.csproj -c Development -p Aut
 if errorlevel 1 goto Error_UATCompileFailed
 
 echo Publishing AutomationTool...
-if exist ..\Binaries\DotNET\AutomationTool (
-	del /f /q ..\Binaries\DotNET\AutomationTool\* >nul 2>&1
-	rmdir /s /q ..\Binaries\DotNET\AutomationTool\runtimes >nul 2>&1
-)
+if exist ..\Binaries\DotNET\AutomationTool del /f /q ..\Binaries\DotNET\AutomationTool\* >nul 2>&1
 if not exist ..\Binaries\DotNET\AutomationTool mkdir ..\Binaries\DotNET\AutomationTool >nul 2>&1
 dotnet publish Programs\AutomationTool\AutomationTool.csproj -c Development -p AutomationToolProjectOnly=true --output ..\Binaries\DotNET\AutomationTool --no-build -v %MSBUILD_LOGLEVEL%
 if errorlevel 1 goto Error_UATCompileFailed
diff --git a/Engine/Build/BatchFiles/BuildUAT.sh b/Engine/Build/BatchFiles/BuildUAT.sh
--- a/Engine/Build/BatchFiles/BuildUAT.sh
+++ b/Engine/Build/BatchFiles/BuildUAT.sh
@@ -68,8 +68,7 @@ if [ $PERFORM_REBUILD -eq 1 ]; then
 
   echo "Publishing AutomationTool..."
   if [ -d ../Binaries/DotNET/AutomationTool ]; then
-    rm -f ../Binaries/DotNET/AutomationTool/* >/dev/null
-    rm -rf ../Binaries/DotNET/AutomationTool/runtimes >/dev/null
+    rm -f ../Binaries/DotNET/AutomationTool/* > /dev/null 2>&1
   fi
   if [ ! -d ../Binaries/DotNET/AutomationTool ]; then
     mkdir -p ../Binaries/DotNET/AutomationTool >/dev/null
diff --git a/Engine/Build/BatchFiles/BuildUBT.bat b/Engine/Build/BatchFiles/BuildUBT.bat
--- a/Engine/Build/BatchFiles/BuildUBT.bat
+++ b/Engine/Build/BatchFiles/BuildUBT.bat
@@ -52,7 +52,7 @@ dotnet build Programs\UnrealBuildTool\UnrealBuildTool.csproj -c Development -v %
 if errorlevel 1 goto Error_UBTCompileFailed
 
 echo Publishing UnrealBuildTool...
-if exist ..\Binaries\DotNET\UnrealBuildTool rmdir /s /q ..\Binaries\DotNET\UnrealBuildTool >nul 2>&1
+if exist ..\Binaries\DotNET\UnrealBuildTool del /f /q ..\Binaries\DotNET\UnrealBuildTool\* >nul 2>&1
 if not exist ..\Binaries\DotNET\UnrealBuildTool mkdir ..\Binaries\DotNET\UnrealBuildTool >nul 2>&1
 dotnet publish Programs\UnrealBuildTool\UnrealBuildTool.csproj -c Development --output ..\Binaries\DotNET\UnrealBuildTool --no-build -v %MSBUILD_LOGLEVEL%
 if errorlevel 1 goto Error_UBTCompileFailed
diff --git a/Engine/Build/BatchFiles/BuildUBT.sh b/Engine/Build/BatchFiles/BuildUBT.sh
--- a/Engine/Build/BatchFiles/BuildUBT.sh
+++ b/Engine/Build/BatchFiles/BuildUBT.sh
@@ -69,7 +69,7 @@ if [ $PERFORM_REBUILD -eq 1 ]; then
 
   echo "Publishing UnrealBuildTool..."
   if [ -d ../Binaries/DotNET/UnrealBuildTool ]; then
-    rm -rf ../Binaries/DotNET/UnrealBuildTool >/dev/null
+    rm -f ../Binaries/DotNET/UnrealBuildTool/* > /dev/null 2>&1
   fi
   if [ ! -d ../Binaries/DotNET/UnrealBuildTool ]; then
     mkdir -p ../Binaries/DotNET/UnrealBuildTool >/dev/null
