Update iOS app extension support for the Background Assets extension
- update ResignApp script to support multiple profiles when app has >1 extension
- support app extensions being in the "Extensions" dir in the bundle instead of the "PlugIns" dir
- add iOS support for accessing files in a Shared File Container

diff --git a/Engine/Build/BatchFiles/Mac/ResignApp.sh b/Engine/Build/BatchFiles/Mac/ResignApp.sh
--- a/Engine/Build/BatchFiles/Mac/ResignApp.sh
+++ b/Engine/Build/BatchFiles/Mac/ResignApp.sh
@@ -1,8 +1,9 @@
-# !/bin/bash
+#!/bin/bash
 
 set -e
 
 FILES_TO_COPY=()
+EXT_PROFILES=()
 
 DoWork()
 {
@@ -56,7 +57,7 @@ DoWork()
 	if [[ ! -z "$MARKETING_VERSION" ]]; then
 		oldmarketingversion=$(/usr/libexec/PlistBuddy -c "Print:CFBundleShortVersionString" "$TARGET/Info.plist")	
 		echo "Changing Marketing Version from $oldmarketingversion to $MARKETING_VERSION"
-		/usr/libexec/PlistBuddy -c "Set:CFBundleShortVersionString $BUNDLE" "$TARGET/Info.plist"
+		/usr/libexec/PlistBuddy -c "Set:CFBundleShortVersionString $MARKETING_VERSION" "$TARGET/Info.plist"
 	fi
 
 	if [[ ! -z "$BUNDLE_VERSION" ]]; then
@@ -74,7 +75,7 @@ DoWork()
     do
         COPY_FILE_NAME=$(basename ${COPY_FILE_PATH});
         echo "Copying $COPY_FILE_PATH to $TARGET/$COPY_FILE_NAME"
-        cp $COPY_FILE_PATH "$TARGET/$COPY_FILE_NAME"
+        cp "$COPY_FILE_PATH" "$TARGET/$COPY_FILE_NAME"
     done
 
 	echo "------------------------------------------------------------------------------"
@@ -92,6 +93,15 @@ DoWork()
 			fi    
 	
 			if [[ -f "$line/embedded.mobileprovision" ]]; then
+				# find the correct mobileprovision file
+				for profile in "${!EXT_PROFILES[@]}"; do
+					if [[ ! -z "$BUNDLE" ]] && [[ "$line" == *".appex"* ]] && [[ "$line" == *"$profile"* ]]; then
+						PROFILE_PATH="${EXT_PROFILES[$profile]}"
+						echo "Copying appex $PROFILE_PATH profile into $line.."
+						cp "$PROFILE_PATH" "$line/embedded.mobileprovision"
+					fi
+				done
+
 				if [[ ! -z "$BUNDLE" ]] && [[ "$line" == *".appex"* ]] && [[ ! -z "$MOBILEAPPEXPROV" ]]; then
 					echo "Copying appex $MOBILEAPPEXPROV profile into $line.."
 					cp "$MOBILEAPPEXPROV" "$line/embedded.mobileprovision"
@@ -193,7 +203,11 @@ Help()
 	echo "  -p | --provision <path to .mobileprovision>"
 	echo "     Path to the .mobileprovision file to sign with"
 	echo "  -pe | --provisionappex <path to .mobileprovision>"
-	echo "     Path to the .mobileprovision file for the app extension to sign with"
+	echo "     Path to the .mobileprovision file for the app extension to sign"
+	echo "  -px <ExtBundleName>=<path to .mobileprovision>"
+	echo "     Path to the .mobileprovision file for the app extension to sign to specified extension name under Build/IOS/Extensions"
+	echo "     Can be repeated for each extension within the app"
+	echo "       ie: '-px NotifService=[REDACTED_USER_HOME]/notifservice.mobileprovision'"
 	echo "  -b | --bundleid"
 	echo "     Bundle ID to use in the app, repleacing existing bundle ID"
 	echo "  -mv | --marketingversion"
@@ -234,6 +248,20 @@ while [[ $# -gt 0 ]]; do
 	  ;;
 	-pe|--provisionappex)
 	  MOBILEAPPEXPROV="$2"
+	  echo "App extension param '-pe' is deprecated. Please use '-px'"
+	  shift # past argument
+	  shift # past value
+	  ;;
+	-px)
+	  EXT_TO_PROFILE="$2"
+	  if [[ $EXT_TO_PROFILE != *=* ]]; then
+		  echo "Invalid extension mobileprovision param. Aborting..."
+		  exit 1
+	  fi
+	  key=${EXT_TO_PROFILE%%=*}
+	  val=${EXT_TO_PROFILE#*=}
+	  EXT_PROFILES["$key"]="$val"
+	  printf 'key=[%s] val=[%s]\n' "${key}" "${val}"
 	  shift # past argument
 	  shift # past value
 	  ;;
diff --git a/Engine/Source/Programs/IOS/iPhonePackager/CodeSigning.cs b/Engine/Source/Programs/IOS/iPhonePackager/CodeSigning.cs
--- a/Engine/Source/Programs/IOS/iPhonePackager/CodeSigning.cs
+++ b/Engine/Source/Programs/IOS/iPhonePackager/CodeSigning.cs
@@ -590,7 +590,7 @@ namespace iPhonePackager
 			Rules2.Add(".*\\.dSYM($|/)", CreateResource(11));
 			Rules2.Add("^", CreateResource(20));
 			Rules2.Add("^(.*/)?\\.DS_Store$", CreateOmittedResource(2000));
-			Rules2.Add("^(Frameworks|SharedFrameworks|PlugIns|Plug-ins|XPCServices|Helpers|MacOS|Library/(Automator|Spotlight|LoginItems))/", CreateNestedResource(10));
+			Rules2.Add("^(Frameworks|SharedFrameworks|PlugIns|Extensions|Plug-ins|XPCServices|Helpers|MacOS|Library/(Automator|Spotlight|LoginItems))/", CreateNestedResource(10));
 			Rules2.Add("^.*", true);
 			Rules2.Add("^.*\\.lproj/", CreateOptionalResource(1000));
 			Rules2.Add("^.*\\.lproj/locversion.plist$", CreateOmittedResource(1100));
diff --git a/Engine/Source/Programs/UnrealBuildTool/ProjectFiles/Xcode/XcodeProject.cs b/Engine/Source/Programs/UnrealBuildTool/ProjectFiles/Xcode/XcodeProject.cs
--- a/Engine/Source/Programs/UnrealBuildTool/ProjectFiles/Xcode/XcodeProject.cs
+++ b/Engine/Source/Programs/UnrealBuildTool/ProjectFiles/Xcode/XcodeProject.cs
@@ -1011,15 +1011,26 @@ namespace UnrealBuildTool.XcodeProjectXcconfig
 	class XcodeCopyAppExtFilesBuildPhase : XcodeBuildPhase
 	{
 		private XcodeFileCollection FileCollection;
+		private ExtensionConfig ExtConfig;
 
-		public XcodeCopyAppExtFilesBuildPhase(XcodeFileCollection FileCollection)
+		public XcodeCopyAppExtFilesBuildPhase(XcodeFileCollection FileCollection, ExtensionConfig ExtConfig)
 			: base("Embed Foundation Extensions", "PBXCopyFilesBuildPhase")
 		{
 			this.FileCollection = FileCollection;
+			this.ExtConfig = ExtConfig ?? new ExtensionConfig();
+
+			string ExtType = this.ExtConfig.ExtensionType ?? "";
+			if (ExtType == "BackgroundAssets")
+			{
+				MiscItems.Add($"dstPath = \"${{EXTENSIONS_FOLDER_PATH}}\";");
+				MiscItems.Add($"dstSubfolderSpec = 16;");
+			}
+			else
+			{
+				MiscItems.Add($"dstPath = \"\";");
+				MiscItems.Add($"dstSubfolderSpec = 13;");
+			}
 
-			// MiscItems.Add($"dstPath = \"${{EXTENSIONS_FOLDER_PATH}}\";");
-			MiscItems.Add($"dstPath = \"\";");
-			MiscItems.Add($"dstSubfolderSpec = 13;");
 			MiscItems.Add($"name = \"{Name}\";");
 		}
 
@@ -1086,6 +1097,7 @@ namespace UnrealBuildTool.XcodeProjectXcconfig
 				Content.WriteLine(3, $"baseConfigurationReference = {Xcconfig.Guid} /* {Xcconfig.Name}.xcconfig */;");
 			}
 			Content.WriteLine(3, "buildSettings = {");
+			Content.WriteLine(4, "SWIFT_VERSION = 5.0;");
 			if (bIncludeAllPlatforms)
 			{
 				Content.WriteLine(4, $"SUPPORTED_PLATFORMS = \"macosx iphonesimulator iphoneos appletvsimulator appletvos xros xrsimulator\";");
@@ -1441,7 +1453,7 @@ namespace UnrealBuildTool.XcodeProjectXcconfig
 							String AppExtPlist = DI.FullName + "/Info.plist";
 							PListHelper PlistDoc = PListHelper.CreateFromFile(AppExtPlist);
 							
-							XcodeCopyAppExtFilesBuildPhase EmbedAppExt = new XcodeCopyAppExtFilesBuildPhase(Project.FileCollection);
+							XcodeCopyAppExtFilesBuildPhase EmbedAppExt = new XcodeCopyAppExtFilesBuildPhase(Project.FileCollection, ExtConfig);
 							
 							// Get the source files
 							foreach (FileSystemInfo FSI in DI.EnumerateFileSystemInfos())
diff --git a/Engine/Source/Runtime/Core/Private/IOS/IOSPlatformFile.cpp b/Engine/Source/Runtime/Core/Private/IOS/IOSPlatformFile.cpp
--- a/Engine/Source/Runtime/Core/Private/IOS/IOSPlatformFile.cpp
+++ b/Engine/Source/Runtime/Core/Private/IOS/IOSPlatformFile.cpp
@@ -1038,6 +1038,23 @@ FString FIOSPlatformFile::ConvertToPlatformPath(const FString& Filename, bool bF
 	return Result;
 }
 
+FString FIOSPlatformFile::GetSharedContainerURL(const TCHAR* AppGroupID)
+{
+	SCOPED_AUTORELEASE_POOL;
+	
+	NSString* GroupID = FString(AppGroupID).GetNSString();
+	NSURL *ContainerUrl = [[NSFileManager defaultManager]
+				  containerURLForSecurityApplicationGroupIdentifier:GroupID];
+	
+	if (ContainerUrl != nil)
+	{
+		NSLog(@"SharedContainer:%@", [ContainerUrl path]);
+		return FString([ContainerUrl path]);
+	}
+	
+	NSLog(@"No SharedContainer with AppGroupID:%@", GroupID);
+	return ""; 
+}
 
 /**
  * iOS platform file declaration
diff --git a/Engine/Source/Runtime/Core/Public/IOS/IOSPlatformFile.h b/Engine/Source/Runtime/Core/Public/IOS/IOSPlatformFile.h
--- a/Engine/Source/Runtime/Core/Public/IOS/IOSPlatformFile.h
+++ b/Engine/Source/Runtime/Core/Public/IOS/IOSPlatformFile.h
@@ -53,6 +53,10 @@ public:
 	virtual bool DoesCreatePublicFiles() override;
 	virtual void SetCreatePublicFiles(bool bCreatePublicFilesIn) override;
 	
+	// If app has a Shared File Container via the AppGroups capability, return the
+	// absolute url path or an empty string if it's not found
+	static FString GetSharedContainerURL(const TCHAR* AppGroupID);
+	
 private:
 	bool IterateDirectoryCommon(const TCHAR* Directory, const TFunctionRef<bool(struct dirent*)>& Visitor);
 	virtual FString ConvertToPlatformPath(const FString& Filename, bool bForWrite, bool bIsPublicWrite) override;
