UAT/UBT: Make out-of-date DotnetDepends checks more robust when launching multiple instances on Windows.
The tools themselves check a mutex, but the bat files could try to rebuild a tool by before that point.
Add an output line when the tool is found to be up-to-date.
When a tool is rebuilt, save a backup (dep.csv.bak) of the previous dep.csv file for later diffing.
Change format of dep.csv file from "hash,path" to "path,hash" to improve the diffing experience.
Add %random%%random% with retries to the temp file to reduce the race risk between multiple instances.
Create the temp files and temp dir immediately after the exist checks to reduce the race risk between multiple instances.

diff --git a/Engine/Build/BatchFiles/BuildUAT.bat b/Engine/Build/BatchFiles/BuildUAT.bat
--- a/Engine/Build/BatchFiles/BuildUAT.bat
+++ b/Engine/Build/BatchFiles/BuildUAT.bat
@@ -26,7 +26,10 @@ if errorlevel 1 goto Error_UBTCompileFailed
 
 rem ## Properties
 set DEPENDS_FILE=..\Intermediate\Build\AutomationTool.dep.csv 
-set TEMP_DEPENDS_FILE=..\Intermediate\Build\AutomationTool.dep.csv.tmp
+set BAK_DEPENDS_FILE=..\Intermediate\Build\AutomationTool.dep.bak.csv
+:tempRetry
+set TEMP_DEPENDS_FILE=..\Intermediate\Build\AutomationTool.dep.tmp-%random%%random%.csv
+if exist "%TEMP_DEPENDS_FILE%" goto :tempRetry
 
 rem ## Command line arguments
 set MSBUILD_LOGLEVEL=%1
@@ -46,6 +49,7 @@ if errorlevel 1 goto Error_UATDependsFailed
 if not exist ..\Binaries\DotNET\AutomationTool\AutomationTool.dll goto Build_AutomationTool
 rem per https://ss64.com/nt/fc.html using redirection syntax rather than errorlevel, based on observed inconsistent results from this function
 fc /C "%TEMP_DEPENDS_FILE%" %DEPENDS_FILE% >nul 2>&1 && (
+	echo AutomationTool is up-to-date
 	set RUNUAT_EXITCODE=0
 	goto Exit
 )
@@ -62,7 +66,8 @@ dotnet publish Programs\AutomationTool\AutomationTool.csproj -c Development -p A
 if errorlevel 1 goto Error_UATCompileFailed
 
 rem record input files - regardless of how we got here, these are now our point of reference
-move /y "%TEMP_DEPENDS_FILE%" %DEPENDS_FILE% >nul
+copy /y "%DEPENDS_FILE%" "%BAK_DEPENDS_FILE%" >nul
+move /y "%TEMP_DEPENDS_FILE%" "%DEPENDS_FILE%" >nul
 
 goto Exit
 
diff --git a/Engine/Build/BatchFiles/BuildUBT.bat b/Engine/Build/BatchFiles/BuildUBT.bat
--- a/Engine/Build/BatchFiles/BuildUBT.bat
+++ b/Engine/Build/BatchFiles/BuildUBT.bat
@@ -22,7 +22,11 @@ if errorlevel 1 goto Error_NoDotnetSDK
 
 rem ## Properties
 set DEPENDS_FILE=..\Intermediate\Build\UnrealBuildTool.dep.csv
-set TEMP_DEPENDS_FILE=..\Intermediate\Build\UnrealBuildTool.dep.csv.tmp
+set BAK_DEPENDS_FILE=..\Intermediate\Build\UnrealBuildTool.dep.bak.csv
+:tempRetry
+set TEMP_DEPENDS_FILE=..\Intermediate\Build\UnrealBuildTool.dep.tmp-%random%%random%.csv
+if exist "%TEMP_DEPENDS_FILE%" goto :tempRetry
+echo >"%TEMP_DEPENDS_FILE%"
 
 rem ## Command line arguments
 set MSBUILD_LOGLEVEL=%1
@@ -42,6 +46,7 @@ if errorlevel 1 goto Error_UBTDependsFailed
 if not exist ..\Binaries\DotNET\UnrealBuildTool\UnrealBuildTool.dll goto Build_UnrealBuildTool
 rem per https://ss64.com/nt/fc.html using redirection syntax rather than errorlevel, based on observed inconsistent results from this function
 fc /C "%TEMP_DEPENDS_FILE%" %DEPENDS_FILE% >nul 2>&1 && (
+	echo UnrealBuildTool is up-to-date
 	set RUNUBT_EXITCODE=0
 	goto Exit
 )
@@ -58,7 +63,8 @@ dotnet publish Programs\UnrealBuildTool\UnrealBuildTool.csproj -c Development --
 if errorlevel 1 goto Error_UBTCompileFailed
 
 rem record input files - regardless of how we got here, these are now our point of reference
-move /y "%TEMP_DEPENDS_FILE%" %DEPENDS_FILE% >nul
+copy /y "%DEPENDS_FILE%" "%BAK_DEPENDS_FILE%" >nul
+move /y "%TEMP_DEPENDS_FILE%" "%DEPENDS_FILE%" >nul
 
 goto Exit
 
diff --git a/Engine/Build/BatchFiles/DotnetDepends.bat b/Engine/Build/BatchFiles/DotnetDepends.bat
--- a/Engine/Build/BatchFiles/DotnetDepends.bat
+++ b/Engine/Build/BatchFiles/DotnetDepends.bat
@@ -27,11 +27,11 @@ if not defined %MSBUILD_LOGLEVEL set MSBUILD_LOGLEVEL=quiet
 
 :tempRetry
 set TEMP_DEPENDS_PATH=%TEMP%\DotnetDepends-%random%%random%
-if exist %TEMP_DEPENDS_PATH% goto :tempRetry
+if exist "%TEMP_DEPENDS_PATH%" goto :tempRetry
+mkdir "%TEMP_DEPENDS_PATH%"
 set TEMP_DEPENDS_FILE=%TEMP_DEPENDS_PATH%\out.csv
 
 rem Check to see if the files in the UnrealBuildTool solution have changed
-mkdir "%TEMP_DEPENDS_PATH%"
 dotnet msbuild "%SLN%" -t:Scan -p:Configuration=Development -p:Platform="Any CPU" -p:OutputPath="%TEMP_DEPENDS_PATH%\\" -p:DependsEncoding=Ascii -noLogo -v:%MSBUILD_LOGLEVEL%
 if errorlevel 1 goto Error_DependsFailed
 copy /b "%TEMP_DEPENDS_PATH%\*.dep.csv" "%TEMP_DEPENDS_FILE%" >nul 2>&1
@@ -48,4 +48,4 @@ goto Exit
 :Exit
 rem ## Clean up temp files
 rmdir /s /q "%TEMP_DEPENDS_PATH%" >nul 2>&1
-exit /B %DEPENDS_EXITCODE%
\ No newline at end of file
+exit /B %DEPENDS_EXITCODE%
