Add a validator to SubmitTool to build UBT in the Analyze configuration to catch newly introduced static analysis issues.

#Synced-CL 48546779

diff --git a/Engine/Build/BatchFiles/RunDotnet.bat b/Engine/Build/BatchFiles/RunDotnet.bat
new file mode 100755
--- /dev/null
+++ b/Engine/Build/BatchFiles/RunDotnet.bat
@@ -0,0 +1,45 @@
+@echo off
+
+rem ## Unreal Engine script to setup and run the bundled dotnet
+rem ## Copyright Epic Games, Inc. All Rights Reserved.
+rem ##
+rem ## This script is expecting to exist in the Engine/Build/BatchFiles directory.  It will not work correctly
+rem ## if you copy it to a different location and run it.
+
+setlocal
+
+rem Change the CWD to /Engine
+pushd "%~dp0..\.."
+
+rem First, make sure the batch file exists in the folder we expect it to.  This is necessary in order to
+rem verify that our relative path to the /Engine directory is correct
+if not exist "Build\BatchFiles\RunDotnet.bat" goto Error_BatchFileInWrongLocation
+
+rem Verify that dotnet is present
+call "%~dp0GetDotnetPath.bat"
+if %ERRORLEVEL% NEQ 0 goto Error_NoDotnetSDK
+
+echo Running dotnet %*
+dotnet %*
+if %ERRORLEVEL% NEQ 0 set RUNDOTNET_EXITCODE=%ERRORLEVEL%
+
+goto Exit
+
+:Error_BatchFileInWrongLocation
+echo.
+echo RunDotnet ERROR: The batch file does not appear to be located in the /Engine/Build/BatchFiles directory.  This script must be run from within that directory.
+echo.
+set RUNDOTNET_EXITCODE=1
+goto Exit
+
+:Error_NoDotnetSDK
+echo.
+echo RunDotnet ERROR: Unable to find a install of Dotnet SDK.  Please make sure you have it installed and that `dotnet` is a globally available command.
+echo.
+set RUNDOTNET_EXITCODE=1
+goto Exit
+
+:Exit
+rem ## Restore original CWD in case we change it
+popd
+exit /B %RUNDOTNET_EXITCODE%
diff --git a/Engine/Build/BatchFiles/RunDotnet.sh b/Engine/Build/BatchFiles/RunDotnet.sh
--- a/Engine/Build/BatchFiles/RunDotnet.sh
+++ b/Engine/Build/BatchFiles/RunDotnet.sh
@@ -1,6 +1,6 @@
 #!/bin/bash
 
-## Unreal Engine dotnet setup script
+## Unreal Engine script to setup and run the bundled dotnet
 ## Copyright Epic Games, Inc. All Rights Reserved.
 
 ## This script is expecting to exist in the UE/Engine/Build/BatchFiles directory.  It will not work correctly
diff --git a/Engine/Source/Programs/SubmitTool/Private/Configuration/Configuration.cpp b/Engine/Source/Programs/SubmitTool/Private/Configuration/Configuration.cpp
--- a/Engine/Source/Programs/SubmitTool/Private/Configuration/Configuration.cpp
+++ b/Engine/Source/Programs/SubmitTool/Private/Configuration/Configuration.cpp
@@ -36,9 +36,11 @@ void FConfiguration::Init()
 	UE_LOG(LogSubmitToolDebug, Log, TEXT("Packaged Root %s"), *FPaths::ConvertRelativePathToFull(FPaths::EngineDir() + TEXT("..")));
 
 #if PLATFORM_WINDOWS
+	Configuration.Values.Add(TEXT("$(ScriptExt)"), TEXT(".bat"));
 	Configuration.Values.Add(TEXT("$(RunUAT)"), BatchFileDir + TEXT("/RunUAT.bat"));
 	Configuration.Values.Add(TEXT("$(RunUBT)"), BatchFileDir + TEXT("/RunUBT.bat"));
 #elif (PLATFORM_LINUX || PLATFORM_MAC)
+	Configuration.Values.Add(TEXT("$(ScriptExt)"), TEXT(".sh"));
 	Configuration.Values.Add(TEXT("$(RunUAT)"), BatchFileDir + TEXT("/RunUAT.sh"));
 	Configuration.Values.Add(TEXT("$(RunUBT)"), BatchFileDir + TEXT("/RunUBT.sh"));
 #else
