Guard rmdir when directory doesn't exist. This was causing the Generate SSH command to emit a 'file not found' error to the console while the subsequent message in the setup states ' Assuming you have seen no errors above, you are ready to compile UE5 remotely....'

#Synced-CL 48480311

diff --git a/Engine/Build/BatchFiles/MakeAndInstallSSHKey.bat b/Engine/Build/BatchFiles/MakeAndInstallSSHKey.bat
--- a/Engine/Build/BatchFiles/MakeAndInstallSSHKey.bat
+++ b/Engine/Build/BatchFiles/MakeAndInstallSSHKey.bat
@@ -49,7 +49,9 @@ echo ===========================================================================
 echo.
 
 pause
-rmdir "%KEY_DIR%" /s/q
+if exist "%KEY_DIR%" (
+	rmdir "%KEY_DIR%" /s/q
+)
 mkdir "%KEY_DIR%" 2> NUL
 pushd %SSHDIR%
 %RSYNC% -az -e '%SSH% -p %SSHPORT%' "%USER%@%MACHINE%":.ssh/RemoteToolChainPrivate.key "%CYGWIN_KEY_PATH%"
