Fix iOS resigning script for multiple app extensions due bash version differences

diff --git a/Engine/Build/BatchFiles/Mac/ResignApp.sh b/Engine/Build/BatchFiles/Mac/ResignApp.sh
--- a/Engine/Build/BatchFiles/Mac/ResignApp.sh
+++ b/Engine/Build/BatchFiles/Mac/ResignApp.sh
@@ -1,10 +1,13 @@
-#!/bin/bash
+#!/bin/zsh
 
 set -e
 
+declare -A EXT_PROFILES
 FILES_TO_COPY=()
 EXT_PROFILES=()
 
+echo "iOS App Resigning Tool: Version 2"
+
 DoWork()
 {
 	echo "------------------------------------------------------------------------------"
@@ -94,7 +97,7 @@ DoWork()
 	
 			if [[ -f "$line/embedded.mobileprovision" ]]; then
 				# find the correct mobileprovision file
-				for profile in "${!EXT_PROFILES[@]}"; do
+				for profile in "${(@k)EXT_PROFILES}"; do
 					if [[ ! -z "$BUNDLE" ]] && [[ "$line" == *".appex"* ]] && [[ "$line" == *"$profile"* ]]; then
 						PROFILE_PATH="${EXT_PROFILES[$profile]}"
 						echo "Copying appex $PROFILE_PATH profile into $line.."
@@ -120,15 +123,15 @@ DoWork()
 				fi    
 				
 				if [[ ! -z "$BUNDLE" ]]; then
-					EXTRAOPTIONS="--entitlements t_entitlements.plist"
+					EXTRAOPTIONS="--entitlements=t_entitlements.plist"
 				else
 					EXTRAOPTIONS="--preserve-metadata=entitlements,flags,identifier"
 				fi
 			fi
 			
 			echo ""
-			echo Codesigning $line [/usr/bin/codesign  --continue -f -s "$DEVELOPER" $EXTRAOPTIONS "$line"]
-			/usr/bin/codesign --continue -f -s "$DEVELOPER" $EXTRAOPTIONS "$line"
+			echo "Codesigning [/usr/bin/codesign --sign \"$DEVELOPER\" --continue --force $EXTRAOPTIONS \"$line\"]"
+			/usr/bin/codesign --sign "$DEVELOPER" --continue --force $EXTRAOPTIONS "$line"
 
 			# reset for next loop
 			EXTRAOPTIONS=""
@@ -260,7 +263,7 @@ while [[ $# -gt 0 ]]; do
 	  fi
 	  key=${EXT_TO_PROFILE%%=*}
 	  val=${EXT_TO_PROFILE#*=}
-	  EXT_PROFILES["$key"]="$val"
+	  EXT_PROFILES[$key]=$val
 	  printf 'key=[%s] val=[%s]\n' "${key}" "${val}"
 	  shift # past argument
 	  shift # past value
