Add option to remove frameworks when resigning an ios build

#rnx
#tests Local and TF upload

diff --git a/Engine/Build/BatchFiles/Mac/ResignApp.sh b/Engine/Build/BatchFiles/Mac/ResignApp.sh
--- a/Engine/Build/BatchFiles/Mac/ResignApp.sh
+++ b/Engine/Build/BatchFiles/Mac/ResignApp.sh
@@ -5,6 +5,7 @@ set -e
 declare -A EXT_PROFILES
 FILES_TO_COPY=()
 EXT_PROFILES=()
+FRAMEWORKS_TO_REMOVE=()
 
 echo "iOS App Resigning Tool: Version 2"
 
@@ -87,6 +88,42 @@ DoWork()
 		EXTRAOPTIONS=""
 		EXTAPPBUNDLEID=""
 		while IFS='' read -r line || [[ -n "$line" ]]; do
+			for FRAMEWORK in "${FRAMEWORKS_TO_REMOVE[@]}"
+			do
+				if [[ "$line" == *"$FRAMEWORK"* ]]; then
+					rm -rf $line
+					continue 2
+				fi
+			done
+			
+			if (( ${#FRAMEWORKS_TO_REMOVE[@]} > 0 )); then
+				if [[ "$line" == *".app" ]]; then
+					EXECUTABLE=$(/usr/libexec/PlistBuddy -c "Print:CFBundleExecutable" "$line/Info.plist")
+					if [[ -z "$EXECUTABLE" ]]; then
+						echo "Error: Could not read CFBundleExecutable from Info.plist."
+						exit 1
+					fi
+					
+					DYLIBS=$(otool -L "$line/$EXECUTABLE" 2>/dev/null)
+					if [[ -z "$DYLIBS" ]]; then
+						echo "Error: otool produced no output. File may not be a Mach-O binary."
+						exit 1
+					fi
+					
+					OTOOL_OUTPUT=("${(f)DYLIBS}")
+					for OTOOL_LINE in "${OTOOL_OUTPUT[@]}"; do
+						for FRAMEWORK in "${FRAMEWORKS_TO_REMOVE[@]}"
+						do
+							if [[ "$OTOOL_LINE" == *"$FRAMEWORK"* ]]; then
+								FRAMEWORK="${OTOOL_LINE%% *}"
+								FRAMEWORK="${FRAMEWORK##[[:space:]]}"
+								FRAMEWORK="${FRAMEWORK%%[[:space:]]}"
+								install_name_tool -change "$FRAMEWORK" / "$TARGET/$EXECUTABLE" 2>/dev/null
+							fi
+						done
+					done
+				fi
+			fi
 	
 			if [[ ! -z "$BUNDLE" ]] && [[ "$line" == *".appex"* ]]; then
 			   extbundleid=$(/usr/libexec/PlistBuddy -c "Print:CFBundleIdentifier" "$line/Info.plist")
@@ -221,6 +258,9 @@ Help()
 	echo "     A commandline to place into the app as uecommandline.txt"
     echo "  -cpy | --copy"
     echo "     Additional file to copy into the app, can be used multiple times. Copied files will be placed into the root directiory, without subdirectories"
+    echo "  -rm | --removeframework <Framework Name>"
+    echo "     Remove the specified framework from the app bundle and from the executable load commands."
+    echo "     Must be used with signing. Can be used multiple times."
 	echo "  -h | --help"
 	echo "     Show this message"
 	exit 0
@@ -297,6 +337,11 @@ while [[ $# -gt 0 ]]; do
       FILES_TO_COPY+=("$2")
       shift # past argument
       shift # past value
+      ;;
+	-rm|--removeframework)
+	  FRAMEWORKS_TO_REMOVE+=("$2")
+      shift # past argument
+      shift # past value
       ;;
 	-h|--help)
 	  Help
