Update message when Engine/Source can not be found

#rnx

#Synced-CL 49553638

diff --git a/Engine/Build/BatchFiles/GenerateProjectFiles.bat b/Engine/Build/BatchFiles/GenerateProjectFiles.bat
--- a/Engine/Build/BatchFiles/GenerateProjectFiles.bat
+++ b/Engine/Build/BatchFiles/GenerateProjectFiles.bat
@@ -11,7 +11,7 @@ echo Setting up Unreal Engine project files...
 
 rem ## First, make sure the batch file exists in the folder we expect it to.  This is necessary in order to
 rem ## verify that our relative path to the /Engine/Source directory is correct
-if not exist "%~dp0..\..\Source" goto Error_BatchFileInWrongLocation
+if not exist "%~dp0..\..\Source" goto Error_EngineSourceMissing
 
 
 rem ## Change the CWD to /Engine/Source.  We always need to run UnrealBuildTool from /Engine/Source!
@@ -41,6 +41,13 @@ rem ## Success!
 popd
 exit /B 0
 
+:Error_EngineSourceMissing
+echo.
+echo GenerateProjectFiles ERROR: The /Engine/Source directory appears to be missing or inaccessable from the location of this batch file.
+echo.
+pause
+goto Exit
+
 :Error_BatchFileInWrongLocation
 echo.
 echo GenerateProjectFiles ERROR: The batch file does not appear to be located in the /Engine/Build/BatchFiles directory.  This script must be run from within that directory.
