Handle old versions of MacOS without /bin/realpath

#rnx
#Synced-CL 49554331

diff --git a/Engine/Build/BatchFiles/Mac/SetupDotnet.sh b/Engine/Build/BatchFiles/Mac/SetupDotnet.sh
--- a/Engine/Build/BatchFiles/Mac/SetupDotnet.sh
+++ b/Engine/Build/BatchFiles/Mac/SetupDotnet.sh
@@ -44,7 +44,10 @@ if [ $IS_DOTNET_INSTALLED -eq 0 ]; then
 	ARCH=x64
 	[ $(uname -m) == "arm64" ] && ARCH=arm64 
 	
-	UE_DOTNET_DIR=$(realpath "$CUR_DIR/../../../Binaries/ThirdParty/DotNet/10.0/mac-$ARCH")
+	UE_DOTNET_DIR=$CUR_DIR/../../../Binaries/ThirdParty/DotNet/10.0/mac-$ARCH
+	if [ $(which realpath) ]; then
+		UE_DOTNET_DIR=$(realpath "$UE_DOTNET_DIR")
+	fi
 	export UE_DOTNET_DIR="$UE_DOTNET_DIR"
 	chmod u+x "$UE_DOTNET_DIR/dotnet"
 	echo $UE_DOTNET_DIR
