Some versions of DYLD error on duplicate framework names, so give each replacement dylib a different name

#rnx
#Synced-CL 49611723

diff --git a/Engine/Build/BatchFiles/Mac/ResignApp.sh b/Engine/Build/BatchFiles/Mac/ResignApp.sh
--- a/Engine/Build/BatchFiles/Mac/ResignApp.sh
+++ b/Engine/Build/BatchFiles/Mac/ResignApp.sh
@@ -87,6 +87,7 @@ DoWork()
 	if [[ ! -z "$DEVELOPER" ]]; then
 		EXTRAOPTIONS=""
 		EXTAPPBUNDLEID=""
+		FRAMEWORK_COUNTER=1
 		while IFS='' read -r line || [[ -n "$line" ]]; do
 			for FRAMEWORK in "${FRAMEWORKS_TO_REMOVE[@]}"
 			do
@@ -114,11 +115,14 @@ DoWork()
 					for OTOOL_LINE in "${OTOOL_OUTPUT[@]}"; do
 						for FRAMEWORK in "${FRAMEWORKS_TO_REMOVE[@]}"
 						do
-							if [[ "$OTOOL_LINE" == *"$FRAMEWORK"* ]]; then
+							if [[ "$OTOOL_LINE" == *"$FRAMEWORK"* && "$OTOOL_LINE" == *" weak"* ]]; then
 								FRAMEWORK="${OTOOL_LINE%% *}"
 								FRAMEWORK="${FRAMEWORK##[[:space:]]}"
 								FRAMEWORK="${FRAMEWORK%%[[:space:]]}"
-								install_name_tool -change "$FRAMEWORK" / "$TARGET/$EXECUTABLE" 2>/dev/null
+								REPLACEMENT_FRAMEWORK="/${FRAMEWORK_COUNTER:#1}"
+								echo "Replacing weak framework "${OTOOL_LINE##[[:space:]]}" with $REPLACEMENT_FRAMEWORK"
+								install_name_tool -change "$FRAMEWORK" "$REPLACEMENT_FRAMEWORK" "$TARGET/$EXECUTABLE" 2>/dev/null
+								((FRAMEWORK_COUNTER++))
 							fi
 						done
 					done
@@ -259,7 +263,7 @@ Help()
     echo "  -cpy | --copy"
     echo "     Additional file to copy into the app, can be used multiple times. Copied files will be placed into the root directiory, without subdirectories"
     echo "  -rm | --removeframework <Framework Name>"
-    echo "     Remove the specified framework from the app bundle and from the executable load commands."
+    echo "     Remove the specified weak framework from the app bundle and from the executable load commands."
     echo "     Must be used with signing. Can be used multiple times."
 	echo "  -h | --help"
 	echo "     Show this message"
